#ue4 ttp332350 LIVE: EDITOR: CRASH: Replacing several actors with a composited Blueprint

Blueprint name is based on a single actor.

diff --git a/Engine/Source/Editor/KismetWidgets/Private/SCreateBlueprintFromActorDialog.cpp b/Engine/Source/Editor/KismetWidgets/Private/SCreateBlueprintFromActorDialog.cpp
--- a/Engine/Source/Editor/KismetWidgets/Private/SCreateBlueprintFromActorDialog.cpp
+++ b/Engine/Source/Editor/KismetWidgets/Private/SCreateBlueprintFromActorDialog.cpp
@@ -60,6 +60,7 @@ void SCreateBlueprintFromActorDialog::Construct( const FArguments& InArgs, TShar
 		{
 			ActorInstanceLabel += Actor->GetActorLabel();
 			ActorInstanceLabel += TEXT( "_" );
+			break; // We don't want any huge names, see ttp332350
 		}
 	}
 	if( ActorInstanceLabel.IsEmpty() )
