[333820] Hide subtrees when empty in Pin/Variable Type menus.

- Modified SPinTypeSelector::Construct() to remove empty top-level items from the tree data source.

diff --git a/Engine/Source/Editor/KismetWidgets/Private/SPinTypeSelector.cpp b/Engine/Source/Editor/KismetWidgets/Private/SPinTypeSelector.cpp
--- a/Engine/Source/Editor/KismetWidgets/Private/SPinTypeSelector.cpp
+++ b/Engine/Source/Editor/KismetWidgets/Private/SPinTypeSelector.cpp
@@ -29,6 +29,24 @@ void SPinTypeSelector::Construct(const FArguments& InArgs, FGetPinTypeTree GetPi
 	TargetPinType = InArgs._TargetPinType;
 
 	GetPinTypeTree.Execute(TypeTreeRoot, bAllowExec, bAllowWildcard);
+
+	// Remove read-only root items if they have no children; there will be no subtree to select non read-only items from in that case
+	int32 RootItemIndex = 0;
+	while(RootItemIndex < TypeTreeRoot.Num())
+	{
+		FPinTypeTreeItem TypeTreeItemPtr = TypeTreeRoot[RootItemIndex];
+		if(TypeTreeItemPtr.IsValid()
+			&& TypeTreeItemPtr->bReadOnly
+			&& TypeTreeItemPtr->Children.Num() == 0)
+		{
+			TypeTreeRoot.RemoveAt(RootItemIndex);
+		}
+		else
+		{
+			++RootItemIndex;
+		}
+	}
+
 	FilteredTypeTreeRoot = TypeTreeRoot;
 
 	this->ChildSlot
