fixed CIS

diff --git a/Engine/Source/Editor/KismetWidgets/Private/SPinTypeSelector.cpp b/Engine/Source/Editor/KismetWidgets/Private/SPinTypeSelector.cpp
--- a/Engine/Source/Editor/KismetWidgets/Private/SPinTypeSelector.cpp
+++ b/Engine/Source/Editor/KismetWidgets/Private/SPinTypeSelector.cpp
@@ -158,8 +158,7 @@ TSharedRef<ITableRow> SPinTypeSelector::GenerateTypeTreeRow(FPinTypeTreeItem InI
 
 	return SNew( SComboRow<FPinTypeTreeItem>, OwnerTree )
 		.ToolTip( IDocumentation::Get()->CreateToolTip( FText::FromString( Tooltip ), NULL, *BigTooltipDocLink, *Description) )
-		.RowContent
-		(
+		[
 			SNew(SHorizontalBox)
 			+SHorizontalBox::Slot()
 			.AutoWidth()
@@ -179,7 +178,7 @@ TSharedRef<ITableRow> SPinTypeSelector::GenerateTypeTreeRow(FPinTypeTreeItem InI
 				.HighlightText(SearchText)
 				.Font( bHasChildren ? FEditorStyle::GetFontStyle(TEXT("Kismet.TypePicker.CategoryFont")) : FEditorStyle::GetFontStyle(TEXT("Kismet.TypePicker.NormalFont")) )
 			]
-		);
+		];
 }
 
 void SPinTypeSelector::OnTypeSelectionChanged(FPinTypeTreeItem Selection, ESelectInfo::Type SelectInfo)
