Animation: Add ability for SScrubWidget/SScrubControlPanel to use full frames in addition to the normal first frame == last frame mode

diff --git a/Engine/Source/Editor/KismetWidgets/Private/SScrubControlPanel.cpp b/Engine/Source/Editor/KismetWidgets/Private/SScrubControlPanel.cpp
--- a/Engine/Source/Editor/KismetWidgets/Private/SScrubControlPanel.cpp
+++ b/Engine/Source/Editor/KismetWidgets/Private/SScrubControlPanel.cpp
@@ -57,9 +57,9 @@ void SScrubControlPanel::Construct( const SScrubControlPanel::FArguments& InArgs
 				.OnCropAnimSequence(InArgs._OnCropAnimSequence)
 				.OnReZeroAnimSequence(InArgs._OnReZeroAnimSequence)
 				.bAllowZoom(InArgs._bAllowZoom)
-				/** Optional, additional values to draw on the timeline **/
 				.DraggableBars(InArgs._DraggableBars)
 				.OnBarDrag(InArgs._OnBarDrag)
+				.bLastFrameIsFirstFrame(InArgs._bLastFrameIsFirstFrame)
 			]
 		]
 
diff --git a/Engine/Source/Editor/KismetWidgets/Private/SScrubWidget.cpp b/Engine/Source/Editor/KismetWidgets/Private/SScrubWidget.cpp
--- a/Engine/Source/Editor/KismetWidgets/Private/SScrubWidget.cpp
+++ b/Engine/Source/Editor/KismetWidgets/Private/SScrubWidget.cpp
@@ -9,20 +9,27 @@
 static const float MinStepLengh=15.f;
 
 /** This function is used by a few random widgets and is mostly arbitrary. It could be moved anywhere. */
-int32 SScrubWidget::GetDivider( float InputMinX, float InputMaxX, FVector2D WidgetSize, float SequenceLength, int32 NumFrames )
+int32 SScrubWidget::GetDivider(float InputMinX, float InputMaxX, FVector2D WidgetSize, float SequenceLength, int32 NumFrames, bool bLastFrameIsFirstFrame)
 {
-	check (NumFrames!=0);
-
 	FTrackScaleInfo TimeScaleInfo(InputMinX, InputMaxX, 0.f, 0.f, WidgetSize);
 
-	float TimePerKey = (NumFrames>1)? SequenceLength/(float)(NumFrames-1) : 0;
+	float TimePerKey = 0.0f;
+	if (bLastFrameIsFirstFrame)
+	{
+		TimePerKey = (NumFrames > 1) ? SequenceLength / (float)(NumFrames - 1) : 0.0f;
+	}
+	else
+	{
+		TimePerKey = (NumFrames > 0) ? SequenceLength / (float)(NumFrames) : 0.0f;
+	}
+
 	float TotalWidgetWidth = TimeScaleInfo.WidgetSize.X;
-	float NumKeys =  TimeScaleInfo.ViewInputRange/TimePerKey;
-	float KeyWidgetWidth = TotalWidgetWidth/ NumKeys;
+	float NumKeys = TimeScaleInfo.ViewInputRange / TimePerKey;
+	float KeyWidgetWidth = TotalWidgetWidth / NumKeys;
 	int32 Divider = 1; 
-	if ( KeyWidgetWidth > 0 )
+	if (KeyWidgetWidth > 0)
 	{
-		Divider = FMath::Max<int32>(50/KeyWidgetWidth, 1);
+		Divider = FMath::Max<int32>(50 / KeyWidgetWidth, 1);
 	}
 	return Divider;
 }
@@ -53,6 +60,7 @@ void SScrubWidget::Construct( const SScrubWidget::FArguments& InArgs )
 	DraggingBar = false;
 
 	bAllowZoom = InArgs._bAllowZoom;
+	bLastFrameIsFirstFrame = InArgs._bLastFrameIsFirstFrame;
 }
 
 int32 SScrubWidget::OnPaint( const FGeometry& AllottedGeometry, const FSlateRect& MyClippingRect, FSlateWindowElementList& OutDrawElements, int32 LayerId, const FWidgetStyle& InWidgetStyle, bool bParentEnabled ) const
@@ -81,12 +89,22 @@ int32 SScrubWidget::OnPaint( const FGeometry& AllottedGeometry, const FSlateRect
 	if ( NumOfKeys.Get() > 0 && SequenceLength.Get() > 0)
 	{
 		FTrackScaleInfo TimeScaleInfo(ViewInputMin.Get(), ViewInputMax.Get(), 0.f, 0.f, AllottedGeometry.Size);
-		int32 Divider = SScrubWidget::GetDivider( ViewInputMin.Get(), ViewInputMax.Get(), AllottedGeometry.Size, SequenceLength.Get(), NumOfKeys.Get() );
+		int32 Divider = SScrubWidget::GetDivider( ViewInputMin.Get(), ViewInputMax.Get(), AllottedGeometry.Size, SequenceLength.Get(), NumOfKeys.Get(), bLastFrameIsFirstFrame );
 		float HalfDivider = Divider/2.f;
 		
 		int32 TotalNumKeys = NumOfKeys.Get();
-		float TimePerKey = (TotalNumKeys > 1)? SequenceLength.Get()/(float)(TotalNumKeys-1) : 0;
-		for ( float KeyVal=0; KeyVal<TotalNumKeys; KeyVal += HalfDivider )
+
+		float TimePerKey = 0.0f;
+		if (bLastFrameIsFirstFrame)
+		{
+			TimePerKey = (TotalNumKeys > 1) ? SequenceLength.Get() / (float)(TotalNumKeys - 1) : 0.0f;
+		}
+		else
+		{
+			TimePerKey = (TotalNumKeys > 0) ? SequenceLength.Get() / (float)(TotalNumKeys) : 0.0f;
+		}
+
+		for (float KeyVal = 0; KeyVal < TotalNumKeys; KeyVal += HalfDivider)
 		{
 			float CurValue = KeyVal*TimePerKey;
 			float XPos = TimeScaleInfo.InputToLocalX(CurValue);
diff --git a/Engine/Source/Editor/KismetWidgets/Public/SScrubControlPanel.h b/Engine/Source/Editor/KismetWidgets/Public/SScrubControlPanel.h
--- a/Engine/Source/Editor/KismetWidgets/Public/SScrubControlPanel.h
+++ b/Engine/Source/Editor/KismetWidgets/Public/SScrubControlPanel.h
@@ -31,6 +31,7 @@ public:
 		, _ViewInputMax()
 		, _OnSetInputViewRange()
 		, _OnCropAnimSequence()
+		, _bLastFrameIsFirstFrame(true)
 	{}
 		SLATE_ATTRIBUTE( float, Value )
 		SLATE_ATTRIBUTE( uint32, NumOfKeys )
@@ -67,6 +68,7 @@ public:
 		/** Optional, additional values to draw on the timeline **/
 		SLATE_ATTRIBUTE( TArray<float>, DraggableBars )
 		SLATE_EVENT( FOnScrubBarDrag, OnBarDrag)
+		SLATE_ARGUMENT( bool, bLastFrameIsFirstFrame )
 	SLATE_END_ARGS()
 
 	/**
diff --git a/Engine/Source/Editor/KismetWidgets/Public/SScrubWidget.h b/Engine/Source/Editor/KismetWidgets/Public/SScrubWidget.h
--- a/Engine/Source/Editor/KismetWidgets/Public/SScrubWidget.h
+++ b/Engine/Source/Editor/KismetWidgets/Public/SScrubWidget.h
@@ -30,6 +30,7 @@ public:
 		, _ViewInputMin()
 		, _OnSetInputViewRange()
 		, _OnCropAnimSequence()
+		, _bLastFrameIsFirstFrame(true)
 		{}
 		
 		/** The value to display */
@@ -54,7 +55,7 @@ public:
 		/** Optional, additional values to draw on the timeline **/
 		SLATE_ATTRIBUTE( TArray<float>, DraggableBars )
 		SLATE_EVENT( FOnScrubBarDrag, OnBarDrag)
-
+		SLATE_ARGUMENT( bool, bLastFrameIsFirstFrame )
 	SLATE_END_ARGS()
 
 
@@ -74,7 +75,7 @@ public:
 	virtual FReply OnMouseWheel( const FGeometry& MyGeometry, const FPointerEvent& MouseEvent ) override;
 	virtual FCursorReply OnCursorQuery( const FGeometry& MyGeometry, const FPointerEvent& CursorEvent ) const override;
 
-	static int32 GetDivider( float InputMinX, float InputMaxX, FVector2D WidgetSize, float SequenceLength, int32 NumFrames );
+	static int32 GetDivider(float InputMinX, float InputMaxX, FVector2D WidgetSize, float SequenceLength, int32 NumFrames, bool bLastFrameIsFirstFrame = true);
 protected:
 
 	/**
@@ -127,5 +128,7 @@ private:
 	/** Bar Dragging*/
 	int32 DraggableBarIndex;
 	bool DraggingBar;
+
+	bool bLastFrameIsFirstFrame;
 };
 
