Animation: Prevent sequence crop menu from showing up if there is no cropping delegate bound

diff --git a/Engine/Source/Editor/KismetWidgets/Private/SScrubWidget.cpp b/Engine/Source/Editor/KismetWidgets/Private/SScrubWidget.cpp
--- a/Engine/Source/Editor/KismetWidgets/Private/SScrubWidget.cpp
+++ b/Engine/Source/Editor/KismetWidgets/Private/SScrubWidget.cpp
@@ -421,7 +421,7 @@ FCursorReply SScrubWidget::OnCursorQuery( const FGeometry& MyGeometry, const FPo
 
 void SScrubWidget::CreateContextMenu(float CurrentFrameTime)
 {
-	if( SequenceLength.Get() >=  MINIMUM_ANIMATION_LENGTH )
+	if (OnCropAnimSequence.IsBound() && (SequenceLength.Get() >= MINIMUM_ANIMATION_LENGTH))
 	{
 		const bool CloseAfterSelection = true;
 		FMenuBuilder MenuBuilder( CloseAfterSelection, NULL );
