Animation: Split out context menu stuff so OnReZero can still be specified independently (longer term both actions need to move out into a GetContextMenu delegate that is passed the current frame time/etc... instead)

diff --git a/Engine/Source/Editor/KismetWidgets/Private/SScrubWidget.cpp b/Engine/Source/Editor/KismetWidgets/Private/SScrubWidget.cpp
--- a/Engine/Source/Editor/KismetWidgets/Private/SScrubWidget.cpp
+++ b/Engine/Source/Editor/KismetWidgets/Private/SScrubWidget.cpp
@@ -421,7 +421,7 @@ FCursorReply SScrubWidget::OnCursorQuery( const FGeometry& MyGeometry, const FPo
 
 void SScrubWidget::CreateContextMenu(float CurrentFrameTime)
 {
-	if (OnCropAnimSequence.IsBound() && (SequenceLength.Get() >= MINIMUM_ANIMATION_LENGTH))
+	if ((OnCropAnimSequence.IsBound() || OnReZeroAnimSequence.IsBound()) && (SequenceLength.Get() >= MINIMUM_ANIMATION_LENGTH))
 	{
 		const bool CloseAfterSelection = true;
 		FMenuBuilder MenuBuilder( CloseAfterSelection, NULL );
@@ -434,37 +434,43 @@ void SScrubWidget::CreateContextMenu(float CurrentFrameTime)
 			FUIAction Action;
 			FText Label;
 
-			//Menu - "Remove Before"
-			//Only show this option if the selected frame is greater than frame 1 (first frame)
-			if( CurrentFrameNumber > 0 )
+			if (OnCropAnimSequence.IsBound())
 			{
-				CurrentFrameFraction = float(CurrentFrameNumber) / (float)NumOfKeys.Get();
-
-				//Corrected frame time based on selected frame number
-				float CorrectedFrameTime = CurrentFrameFraction * SequenceLength.Get();
-
-				Action = FUIAction( FExecuteAction::CreateSP( this, &SScrubWidget::OnSequenceCropped, true, CorrectedFrameTime ) );
-				Label = FText::Format( LOCTEXT("RemoveTillFrame", "Remove till frame {0}"), FText::AsNumber(CurrentFrameNumber) );
-				MenuBuilder.AddMenuEntry( Label, LOCTEXT("RemoveBefore_ToolTip", "Remove sequence before current position"), FSlateIcon(), Action);
+				//Menu - "Remove Before"
+				//Only show this option if the selected frame is greater than frame 1 (first frame)
+				if (CurrentFrameNumber > 0)
+				{
+					CurrentFrameFraction = float(CurrentFrameNumber) / (float)NumOfKeys.Get();
+
+					//Corrected frame time based on selected frame number
+					float CorrectedFrameTime = CurrentFrameFraction * SequenceLength.Get();
+
+					Action = FUIAction(FExecuteAction::CreateSP(this, &SScrubWidget::OnSequenceCropped, true, CorrectedFrameTime));
+					Label = FText::Format(LOCTEXT("RemoveTillFrame", "Remove till frame {0}"), FText::AsNumber(CurrentFrameNumber));
+					MenuBuilder.AddMenuEntry(Label, LOCTEXT("RemoveBefore_ToolTip", "Remove sequence before current position"), FSlateIcon(), Action);
+				}
+
+				uint32 NextFrameNumber = CurrentFrameNumber + 1;
+
+				//Menu - "Remove After"
+				//Only show this option if next frame (CurrentFrameNumber + 1) is valid
+				if (NextFrameNumber < NumOfKeys.Get())
+				{
+					float NextFrameFraction = float(NextFrameNumber) / (float)NumOfKeys.Get();
+					float NextFrameTime = NextFrameFraction * SequenceLength.Get();
+					Action = FUIAction(FExecuteAction::CreateSP(this, &SScrubWidget::OnSequenceCropped, false, NextFrameTime));
+					Label = FText::Format(LOCTEXT("RemoveFromFrame", "Remove from frame {0}"), FText::AsNumber(NextFrameNumber));
+					MenuBuilder.AddMenuEntry(Label, LOCTEXT("RemoveAfter_ToolTip", "Remove sequence after current position"), FSlateIcon(), Action);
+				}
 			}
 
-			uint32 NextFrameNumber = CurrentFrameNumber + 1;
-
-			//Menu - "Remove After"
-			//Only show this option if next frame (CurrentFrameNumber + 1) is valid
-			if( NextFrameNumber < NumOfKeys.Get() )
+			if (OnReZeroAnimSequence.IsBound())
 			{
-				float NextFrameFraction = float(NextFrameNumber) / (float)NumOfKeys.Get();
-				float NextFrameTime = NextFrameFraction * SequenceLength.Get();
-				Action = FUIAction( FExecuteAction::CreateSP( this, &SScrubWidget::OnSequenceCropped, false, NextFrameTime ) );
-				Label = FText::Format( LOCTEXT("RemoveFromFrame", "Remove from frame {0}"), FText::AsNumber(NextFrameNumber) );
-				MenuBuilder.AddMenuEntry( Label, LOCTEXT("RemoveAfter_ToolTip", "Remove sequence after current position"), FSlateIcon(), Action);
+				//Menu - "ReZero"
+				Action = FUIAction(FExecuteAction::CreateSP(this, &SScrubWidget::OnReZero));
+				Label = FText::Format(LOCTEXT("ReZeroAtFrame", "ReZero at frame {0}"), FText::AsNumber(CurrentFrameNumber));
+				MenuBuilder.AddMenuEntry(Label, LOCTEXT("ReZeroAtFrame_ToolTip", "ReZero sequence at the current frame"), FSlateIcon(), Action);
 			}
-
-			//Menu - "ReZero"
-			Action = FUIAction( FExecuteAction::CreateSP( this, &SScrubWidget::OnReZero ) );
-			Label = FText::Format( LOCTEXT("ReZeroAtFrame", "ReZero at frame {0}"), FText::AsNumber(CurrentFrameNumber) );
-			MenuBuilder.AddMenuEntry( Label, LOCTEXT("ReZeroAtFrame_ToolTip", "ReZero sequence at the current frame"), FSlateIcon(), Action);
 		}
 		MenuBuilder.EndSection();
 
