Renamed/added tool tip for scrub menu items.

diff --git a/Engine/Source/Editor/KismetWidgets/Private/SScrubWidget.cpp b/Engine/Source/Editor/KismetWidgets/Private/SScrubWidget.cpp
--- a/Engine/Source/Editor/KismetWidgets/Private/SScrubWidget.cpp
+++ b/Engine/Source/Editor/KismetWidgets/Private/SScrubWidget.cpp
@@ -429,7 +429,7 @@ void SScrubWidget::CreateContextMenu(float CurrentFrameTime)
 					float CorrectedFrameTime = CurrentFrameFraction * SequenceLength.Get();
 
 					Action = FUIAction(FExecuteAction::CreateSP(this, &SScrubWidget::OnSequenceCropped, true, CorrectedFrameTime));
-					Label = FText::Format(LOCTEXT("RemoveTillFrame", "Remove till frame {0}"), FText::AsNumber(CurrentFrameNumber));
+					Label = FText::Format(LOCTEXT("RemoveTillFrame", "Remove frame 0 to frame {0}"), FText::AsNumber(CurrentFrameNumber));
 					MenuBuilder.AddMenuEntry(Label, LOCTEXT("RemoveBefore_ToolTip", "Remove sequence before current position"), FSlateIcon(), Action);
 				}
 
@@ -442,7 +442,7 @@ void SScrubWidget::CreateContextMenu(float CurrentFrameTime)
 					float NextFrameFraction = float(NextFrameNumber) / (float)NumOfKeys.Get();
 					float NextFrameTime = NextFrameFraction * SequenceLength.Get();
 					Action = FUIAction(FExecuteAction::CreateSP(this, &SScrubWidget::OnSequenceCropped, false, NextFrameTime));
-					Label = FText::Format(LOCTEXT("RemoveFromFrame", "Remove from frame {0}"), FText::AsNumber(NextFrameNumber));
+					Label = FText::Format(LOCTEXT("RemoveFromFrame", "Remove from frame {0} to frame {1}"), FText::AsNumber(NextFrameNumber), FText::AsNumber(NumOfKeys.Get()));
 					MenuBuilder.AddMenuEntry(Label, LOCTEXT("RemoveAfter_ToolTip", "Remove sequence after current position"), FSlateIcon(), Action);
 				}
 			}
@@ -452,7 +452,7 @@ void SScrubWidget::CreateContextMenu(float CurrentFrameTime)
 				//Menu - "ReZero"
 				Action = FUIAction(FExecuteAction::CreateSP(this, &SScrubWidget::OnReZero));
 				Label = FText::Format(LOCTEXT("ReZeroAtFrame", "ReZero at frame {0}"), FText::AsNumber(CurrentFrameNumber));
-				MenuBuilder.AddMenuEntry(Label, LOCTEXT("ReZeroAtFrame_ToolTip", "ReZero sequence at the current frame"), FSlateIcon(), Action);
+				MenuBuilder.AddMenuEntry(Label, LOCTEXT("ReZeroAtFrame_ToolTip", "Resets the root track of the frame to (0, 0, 0), and apply the difference to all root transform of the sequence. It moves whole sequence to the amount of current root transform. "), FSlateIcon(), Action);
 			}
 		}
 		MenuBuilder.EndSection();
