Added handling for NULL RootWindow to FCreateBlueprintFromActorDialog::OpenDialog() and UEditorEngine::ConvertActors()

diff --git a/Engine/Source/Editor/KismetWidgets/Private/CreateBlueprintFromActorDialog.cpp b/Engine/Source/Editor/KismetWidgets/Private/CreateBlueprintFromActorDialog.cpp
--- a/Engine/Source/Editor/KismetWidgets/Private/CreateBlueprintFromActorDialog.cpp
+++ b/Engine/Source/Editor/KismetWidgets/Private/CreateBlueprintFromActorDialog.cpp
@@ -27,7 +27,14 @@ void FCreateBlueprintFromActorDialog::OpenDialog(bool bInHarvest)
 		);
 
 	TSharedPtr<SWindow> RootWindow = FGlobalTabmanager::Get()->GetRootWindow();
-	FSlateApplication::Get().AddWindowAsNativeChild(PickBlueprintPathWidget.ToSharedRef(), RootWindow.ToSharedRef() );
+	if (RootWindow.IsValid())
+	{
+		FSlateApplication::Get().AddWindowAsNativeChild(PickBlueprintPathWidget.ToSharedRef(), RootWindow.ToSharedRef());
+	}
+	else
+	{
+		FSlateApplication::Get().AddWindow(PickBlueprintPathWidget.ToSharedRef());
+	}
 }
 
 void FCreateBlueprintFromActorDialog::OnCreateBlueprint(const FString& InAssetPath, bool bInHarvest)
diff --git a/Engine/Source/Editor/UnrealEd/Private/Editor.cpp b/Engine/Source/Editor/UnrealEd/Private/Editor.cpp
--- a/Engine/Source/Editor/UnrealEd/Private/Editor.cpp
+++ b/Engine/Source/Editor/UnrealEd/Private/Editor.cpp
@@ -5631,7 +5631,14 @@ void UEditorEngine::ConvertActors( const TArray<AActor*>& ActorsToConvert, UClas
 			);
 
 		TSharedPtr<SWindow> RootWindow = FGlobalTabmanager::Get()->GetRootWindow();
-		FSlateApplication::Get().AddWindowAsNativeChild(CreateAssetFromActorWindow.ToSharedRef(), RootWindow.ToSharedRef());
+		if (RootWindow.IsValid())
+		{
+			FSlateApplication::Get().AddWindowAsNativeChild(CreateAssetFromActorWindow.ToSharedRef(), RootWindow.ToSharedRef());
+		}
+		else
+		{
+			FSlateApplication::Get().AddWindow(CreateAssetFromActorWindow.ToSharedRef());
+		}
 	}
 	else
 	{
