Fixed SScrubWidget::OnPaint drawing the draggable bars with an incorrect offset

https://udn.unrealengine.com/questions/217968/found-a-small-bug-in-the-scrub-widget-code.html

Also added some extra const, and vastly reduced the number of times the DraggableBars delegate was called.

diff --git a/Engine/Source/Editor/KismetWidgets/Private/SScrubWidget.cpp b/Engine/Source/Editor/KismetWidgets/Private/SScrubWidget.cpp
--- a/Engine/Source/Editor/KismetWidgets/Private/SScrubWidget.cpp
+++ b/Engine/Source/Editor/KismetWidgets/Private/SScrubWidget.cpp
@@ -67,7 +67,7 @@ int32 SScrubWidget::OnPaint( const FPaintArgs& Args, const FGeometry& AllottedGe
 
 	const int32 BackgroundLayer = LayerId;
 
-	FSlateFontInfo SmallLayoutFont( FPaths::EngineContentDir() / TEXT("Slate/Fonts/Roboto-Regular.ttf"), 10 );
+	const FSlateFontInfo SmallLayoutFont( FPaths::EngineContentDir() / TEXT("Slate/Fonts/Roboto-Regular.ttf"), 10 );
 
 	const bool bEnabled = ShouldBeEnabled( bParentEnabled );
 	const ESlateDrawEffect::Type DrawEffects = bEnabled ? ESlateDrawEffect::None : ESlateDrawEffect::DisabledEffect;
@@ -75,27 +75,27 @@ int32 SScrubWidget::OnPaint( const FPaintArgs& Args, const FGeometry& AllottedGe
 	const int32 TextLayer = BackgroundLayer + 1;
 
 	const FSlateBrush* StyleInfo = FEditorStyle::GetBrush( TEXT( "ProgressBar.Background" ) );
-	FSlateRect GeomRect = AllottedGeometry.GetRect();
+	const FSlateRect GeomRect = AllottedGeometry.GetRect();
 
 	if ( NumOfKeys.Get() > 0 && SequenceLength.Get() > 0)
 	{
-		FTrackScaleInfo TimeScaleInfo(ViewInputMin.Get(), ViewInputMax.Get(), 0.f, 0.f, AllottedGeometry.Size);
-		int32 Divider = SScrubWidget::GetDivider( ViewInputMin.Get(), ViewInputMax.Get(), AllottedGeometry.Size, SequenceLength.Get(), NumOfKeys.Get() );
-		float HalfDivider = Divider/2.f;
+		const FTrackScaleInfo TimeScaleInfo(ViewInputMin.Get(), ViewInputMax.Get(), 0.f, 0.f, AllottedGeometry.Size);
+		const int32 Divider = SScrubWidget::GetDivider( ViewInputMin.Get(), ViewInputMax.Get(), AllottedGeometry.Size, SequenceLength.Get(), NumOfKeys.Get() );
+		const float HalfDivider = Divider/2.f;
 		
-		int32 TotalNumKeys = NumOfKeys.Get();
+		const int32 TotalNumKeys = NumOfKeys.Get();
 
-		float TimePerKey = (TotalNumKeys > 0) ? SequenceLength.Get() / (float)(TotalNumKeys) : 0.0f;
+		const float TimePerKey = (TotalNumKeys > 0) ? SequenceLength.Get() / (float)(TotalNumKeys) : 0.0f;
 
 		for (float KeyVal = 0; KeyVal < TotalNumKeys; KeyVal += HalfDivider)
 		{
-			float CurValue = KeyVal*TimePerKey;
-			float XPos = TimeScaleInfo.InputToLocalX(CurValue);
+			const float CurValue = KeyVal*TimePerKey;
+			const float XPos = TimeScaleInfo.InputToLocalX(CurValue);
 
 			if ( FGenericPlatformMath::Fmod(KeyVal, Divider) == 0.f )
 			{
-				FVector2D Offset(XPos, 0.f);
-				FVector2D Size(1, GeomRect.Bottom-GeomRect.Top);
+				const FVector2D Offset(XPos, 0.f);
+				const FVector2D Size(1, GeomRect.Bottom-GeomRect.Top);
 				// draw each box with key frame
 				FSlateDrawElement::MakeBox(
 					OutDrawElements,
@@ -107,12 +107,12 @@ int32 SScrubWidget::OnPaint( const FPaintArgs& Args, const FGeometry& AllottedGe
 					InWidgetStyle.GetColorAndOpacityTint()
 					);
 
-				int32 FrameNumber = KeyVal;
-				FString FrameString = FString::Printf(TEXT("%d"), (FrameNumber));
-				FVector2D TextOffset(XPos+2.f, 0.f);
+				const int32 FrameNumber = KeyVal;
+				const FString FrameString = FString::Printf(TEXT("%d"), (FrameNumber));
+				const FVector2D TextOffset(XPos+2.f, 0.f);
 
 				const TSharedRef< FSlateFontMeasure > FontMeasureService = FSlateApplication::Get().GetRenderer()->GetFontMeasureService();
-				FVector2D TextSize = FontMeasureService->Measure(FrameString, SmallLayoutFont);
+				const FVector2D TextSize = FontMeasureService->Measure(FrameString, SmallLayoutFont);
 
 				FSlateDrawElement::MakeText(
 					OutDrawElements,
@@ -125,9 +125,9 @@ int32 SScrubWidget::OnPaint( const FPaintArgs& Args, const FGeometry& AllottedGe
 			}
  			else if (HalfDivider > 1.f)
  			{
- 				float Height = GeomRect.Bottom-GeomRect.Top;
- 				FVector2D Offset(XPos, Height*0.25f);
- 				FVector2D Size(1, Height*0.5f);
+ 				const float Height = GeomRect.Bottom-GeomRect.Top;
+ 				const FVector2D Offset(XPos, Height*0.25f);
+ 				const FVector2D Size(1, Height*0.5f);
  				// draw each box with key frame
  				FSlateDrawElement::MakeBox(
  					OutDrawElements,
@@ -142,47 +142,51 @@ int32 SScrubWidget::OnPaint( const FPaintArgs& Args, const FGeometry& AllottedGe
  			}
 		}
 
-		const float XPos = TimeScaleInfo.InputToLocalX(ValueAttribute.Get());
-		const float Height = AllottedGeometry.Size.Y;
-		const FVector2D Offset( XPos - Height*0.25f, 0.f );
-
 		const int32 ArrowLayer = TextLayer + 1;
-		FPaintGeometry MyGeometry =	AllottedGeometry.ToPaintGeometry( Offset, FVector2D(Height*0.5f, Height) );
-		FLinearColor ScrubColor = InWidgetStyle.GetColorAndOpacityTint();
-		ScrubColor.A = ScrubColor.A*0.5f;
-		ScrubColor.B *= 0.1f;
-		ScrubColor.G *= 0.1f;
-		FSlateDrawElement::MakeBox( 
-			OutDrawElements,
-			ArrowLayer, 
-			MyGeometry, 
-			StyleInfo, 
-			MyClippingRect, 
-			DrawEffects, 
-			ScrubColor
-			);
-
+		{
+			const float XPos = TimeScaleInfo.InputToLocalX(ValueAttribute.Get());
+			const float Height = AllottedGeometry.Size.Y;
+			const FVector2D Offset( XPos - Height*0.25f, 0.f );
+
+			FPaintGeometry MyGeometry =	AllottedGeometry.ToPaintGeometry( Offset, FVector2D(Height*0.5f, Height) );
+			FLinearColor ScrubColor = InWidgetStyle.GetColorAndOpacityTint();
+			ScrubColor.A = ScrubColor.A*0.5f;
+			ScrubColor.B *= 0.1f;
+			ScrubColor.G *= 0.1f;
+			FSlateDrawElement::MakeBox( 
+				OutDrawElements,
+				ArrowLayer, 
+				MyGeometry, 
+				StyleInfo, 
+				MyClippingRect, 
+				DrawEffects, 
+				ScrubColor
+				);
+		}
 
 		// Draggable Bars
-		for ( int32 BarIdx=0; DraggableBars.IsBound() && BarIdx < DraggableBars.Get().Num(); BarIdx++ )
+		if ( DraggableBars.IsBound() )
 		{
-			float BarXPos = TimeScaleInfo.InputToLocalX(DraggableBars.Get()[BarIdx]);
-			FVector2D BarOffset(BarXPos-2.f, 0.f);
-			FVector2D Size(4.f, GeomRect.Bottom-GeomRect.Top);
+			for ( const float BarValue : DraggableBars.Get() )
+			{
+				const float BarXPos = TimeScaleInfo.InputToLocalX(BarValue);
+				const FVector2D BarOffset(BarXPos-2.f, 0.f);
+				const FVector2D Size(4.f, GeomRect.Bottom-GeomRect.Top);
 
-			FLinearColor BarColor = InWidgetStyle.GetColorAndOpacityTint();	
-			BarColor.R *= 0.1f;
-			BarColor.G *= 0.1f;
+				FLinearColor BarColor = InWidgetStyle.GetColorAndOpacityTint();	
+				BarColor.R *= 0.1f;
+				BarColor.G *= 0.1f;
 
-			FSlateDrawElement::MakeBox(
-				OutDrawElements,
-				ArrowLayer+1,
-				AllottedGeometry.ToPaintGeometry(Offset, Size),
-				StyleInfo,
-				MyClippingRect,
-				DrawEffects,
-				BarColor
-				);
+				FSlateDrawElement::MakeBox(
+					OutDrawElements,
+					ArrowLayer+1,
+					AllottedGeometry.ToPaintGeometry(BarOffset, Size),
+					StyleInfo,
+					MyClippingRect,
+					DrawEffects,
+					BarColor
+					);
+			}
 		}
 
 		return FMath::Max( ArrowLayer, SCompoundWidget::OnPaint( Args, AllottedGeometry, MyClippingRect, OutDrawElements, ArrowLayer, InWidgetStyle, bEnabled ) );
@@ -285,12 +289,16 @@ FReply SScrubWidget::OnMouseMove( const FGeometry& MyGeometry, const FPointerEve
 		FVector2D CursorPos = MyGeometry.AbsoluteToLocal(MouseEvent.GetScreenSpacePosition());
 		FTrackScaleInfo ScaleInfo(ViewInputMin.Get(),  ViewInputMax.Get(), 0.f, 0.f, MyGeometry.Size);
 		DraggableBarIndex = INDEX_NONE;
-		for ( int32 I=0; DraggableBars.IsBound() && I < DraggableBars.Get().Num(); I++ )
+		if ( DraggableBars.IsBound() )
 		{
-			if( FMath::Abs( ScaleInfo.InputToLocalX(DraggableBars.Get()[I]) - CursorPos.X ) < 10 )
+			const TArray<float>& DraggableBarsVal = DraggableBars.Get();
+			for ( int32 I=0; I < DraggableBarsVal.Num(); I++ )
 			{
-				DraggableBarIndex = I;
-				break;
+				if( FMath::Abs( ScaleInfo.InputToLocalX(DraggableBarsVal[I]) - CursorPos.X ) < 10 )
+				{
+					DraggableBarIndex = I;
+					break;
+				}
 			}
 		}
 	}
