Performance optimization: Blueprint pin type selector does not generate the list of items until the list is expanded.

diff --git a/Engine/Source/Editor/KismetWidgets/Private/SPinTypeSelector.cpp b/Engine/Source/Editor/KismetWidgets/Private/SPinTypeSelector.cpp
--- a/Engine/Source/Editor/KismetWidgets/Private/SPinTypeSelector.cpp
+++ b/Engine/Source/Editor/KismetWidgets/Private/SPinTypeSelector.cpp
@@ -29,27 +29,6 @@ void SPinTypeSelector::Construct(const FArguments& InArgs, FGetPinTypeTree GetPi
 
 	TargetPinType = InArgs._TargetPinType;
 
-	GetPinTypeTree.Execute(TypeTreeRoot, bAllowExec, bAllowWildcard);
-
-	// Remove read-only root items if they have no children; there will be no subtree to select non read-only items from in that case
-	int32 RootItemIndex = 0;
-	while(RootItemIndex < TypeTreeRoot.Num())
-	{
-		FPinTypeTreeItem TypeTreeItemPtr = TypeTreeRoot[RootItemIndex];
-		if(TypeTreeItemPtr.IsValid()
-			&& TypeTreeItemPtr->bReadOnly
-			&& TypeTreeItemPtr->Children.Num() == 0)
-		{
-			TypeTreeRoot.RemoveAt(RootItemIndex);
-		}
-		else
-		{
-			++RootItemIndex;
-		}
-	}
-
-	FilteredTypeTreeRoot = TypeTreeRoot;
-
 	this->ChildSlot
 	[
 		SNew(SHorizontalBox)
@@ -249,9 +228,27 @@ void SPinTypeSelector::GetTypeChildren(FPinTypeTreeItem InItem, TArray<FPinTypeT
 
 TSharedRef<SWidget>	SPinTypeSelector::GetMenuContent()
 {
-	// Refresh the type tree, in case we've gotten any new valid types from other blueprints
 	GetPinTypeTree.Execute(TypeTreeRoot, bAllowExec, bAllowWildcard);
 
+	// Remove read-only root items if they have no children; there will be no subtree to select non read-only items from in that case
+	int32 RootItemIndex = 0;
+	while(RootItemIndex < TypeTreeRoot.Num())
+	{
+		FPinTypeTreeItem TypeTreeItemPtr = TypeTreeRoot[RootItemIndex];
+		if(TypeTreeItemPtr.IsValid()
+			&& TypeTreeItemPtr->bReadOnly
+			&& TypeTreeItemPtr->Children.Num() == 0)
+		{
+			TypeTreeRoot.RemoveAt(RootItemIndex);
+		}
+		else
+		{
+			++RootItemIndex;
+		}
+	}
+
+	FilteredTypeTreeRoot = TypeTreeRoot;
+
 	if( !MenuContent.IsValid() )
 	{
 		// Pre-build the tree view and search box as it is needed as a parameter for the context menu's container.
