#AnimEditing: support inserting frame

- This only supports a frame insertion, not resizing length
- Main difficulty is to link with curves that have been modified. We'll need to properly shift whole thing if we want to.

diff --git a/Engine/Source/Editor/KismetWidgets/Private/SScrubControlPanel.cpp b/Engine/Source/Editor/KismetWidgets/Private/SScrubControlPanel.cpp
--- a/Engine/Source/Editor/KismetWidgets/Private/SScrubControlPanel.cpp
+++ b/Engine/Source/Editor/KismetWidgets/Private/SScrubControlPanel.cpp
@@ -56,6 +56,7 @@ void SScrubControlPanel::Construct( const SScrubControlPanel::FArguments& InArgs
 				.ViewInputMax(InArgs._ViewInputMax)
 				.OnSetInputViewRange(InArgs._OnSetInputViewRange)
 				.OnCropAnimSequence(InArgs._OnCropAnimSequence)
+				.OnAddAnimSequence(InArgs._OnAddAnimSequence)
 				.OnReZeroAnimSequence(InArgs._OnReZeroAnimSequence)
 				.bAllowZoom(InArgs._bAllowZoom)
 				.DraggableBars(InArgs._DraggableBars)
diff --git a/Engine/Source/Editor/KismetWidgets/Private/SScrubWidget.cpp b/Engine/Source/Editor/KismetWidgets/Private/SScrubWidget.cpp
--- a/Engine/Source/Editor/KismetWidgets/Private/SScrubWidget.cpp
+++ b/Engine/Source/Editor/KismetWidgets/Private/SScrubWidget.cpp
@@ -40,6 +40,7 @@ void SScrubWidget::Construct( const SScrubWidget::FArguments& InArgs )
 	ViewInputMax = InArgs._ViewInputMax;
 	OnSetInputViewRange = InArgs._OnSetInputViewRange;
 	OnCropAnimSequence = InArgs._OnCropAnimSequence;
+	OnAddAnimSequence = InArgs._OnAddAnimSequence;
 	OnReZeroAnimSequence = InArgs._OnReZeroAnimSequence;
 
 	DraggableBars = InArgs._DraggableBars;
@@ -412,7 +413,7 @@ FCursorReply SScrubWidget::OnCursorQuery( const FGeometry& MyGeometry, const FPo
 
 void SScrubWidget::CreateContextMenu(float CurrentFrameTime)
 {
-	if ((OnCropAnimSequence.IsBound() || OnReZeroAnimSequence.IsBound()) && (SequenceLength.Get() >= MINIMUM_ANIMATION_LENGTH))
+	if ((OnCropAnimSequence.IsBound() || OnReZeroAnimSequence.IsBound() || OnAddAnimSequence.IsBound()) && (SequenceLength.Get() >= MINIMUM_ANIMATION_LENGTH))
 	{
 		const bool CloseAfterSelection = true;
 		FMenuBuilder MenuBuilder( CloseAfterSelection, NULL );
@@ -420,7 +421,7 @@ void SScrubWidget::CreateContextMenu(float CurrentFrameTime)
 		MenuBuilder.BeginSection("SequenceEditingContext", LOCTEXT("SequenceEditing", "Sequence Editing") );
 		{
 			float CurrentFrameFraction = CurrentFrameTime / SequenceLength.Get();
-			uint32 CurrentFrameNumber = CurrentFrameFraction * NumOfKeys.Get();
+			int32 CurrentFrameNumber = CurrentFrameFraction * NumOfKeys.Get();
 
 			FUIAction Action;
 			FText Label;
@@ -455,6 +456,20 @@ void SScrubWidget::CreateContextMenu(float CurrentFrameTime)
 				}
 			}
 
+			if (OnAddAnimSequence.IsBound())
+			{
+				//Corrected frame time based on selected frame number
+				float CorrectedFrameTime = CurrentFrameFraction * SequenceLength.Get();
+
+				Action = FUIAction(FExecuteAction::CreateSP(this, &SScrubWidget::OnSequenceAdded, true, CurrentFrameNumber));
+				Label = FText::Format(LOCTEXT("InsertBeforeCurrentFrame", "Insert frame before {0}"), FText::AsNumber(CurrentFrameNumber));
+				MenuBuilder.AddMenuEntry(Label, LOCTEXT("InsertBefore_ToolTip", "Insert a frame before current position"), FSlateIcon(), Action);
+
+				Action = FUIAction(FExecuteAction::CreateSP(this, &SScrubWidget::OnSequenceAdded, false, CurrentFrameNumber));
+				Label = FText::Format(LOCTEXT("InsertAfterCurrentFrame", "Insert frame after {0}"), FText::AsNumber(CurrentFrameNumber));
+				MenuBuilder.AddMenuEntry(Label, LOCTEXT("InsertAfter_ToolTip", "Insert a frame after current position"), FSlateIcon(), Action);
+			}
+
 			if (OnReZeroAnimSequence.IsBound())
 			{
 				//Menu - "ReZero"
@@ -477,6 +492,16 @@ void SScrubWidget::OnSequenceCropped( bool bFromStart, float CurrentFrameTime )
 	OnSetInputViewRange.ExecuteIfBound( ViewInputMin.Get(), ViewInputMax.Get() );
 }
 
+void SScrubWidget::OnSequenceAdded(bool bBefore, int32 CurrentFrameNumber)
+{
+	OnAddAnimSequence.ExecuteIfBound(bBefore, CurrentFrameNumber);
+
+	//Update scrubs new length to be new Sequence Length
+	// @Todo fixme: this whole thing needs to change to "Refresh" 
+	// - including the OnSequenceCropped
+	OnSetInputViewRange.ExecuteIfBound(ViewInputMin.Get(), SequenceLength.Get());
+}
+
 void SScrubWidget::OnReZero()
 {
 	OnReZeroAnimSequence.ExecuteIfBound();
diff --git a/Engine/Source/Editor/KismetWidgets/Public/SScrubControlPanel.h b/Engine/Source/Editor/KismetWidgets/Public/SScrubControlPanel.h
--- a/Engine/Source/Editor/KismetWidgets/Public/SScrubControlPanel.h
+++ b/Engine/Source/Editor/KismetWidgets/Public/SScrubControlPanel.h
@@ -61,6 +61,8 @@ public:
 		SLATE_EVENT( FOnSetInputViewRange, OnSetInputViewRange )
 		/** Called when an anim sequence is cropped before/after a selected frame */
 		SLATE_EVENT( FOnCropAnimSequence, OnCropAnimSequence )
+		/** Called when an frane is added before/after a selected frame */
+		SLATE_EVENT(FOnAddAnimSequence, OnAddAnimSequence)
 		/** Called to zero out selected frame's translation from origin */
 		SLATE_EVENT( FSimpleDelegate, OnReZeroAnimSequence )
 		SLATE_ATTRIBUTE( bool, IsRealtimeStreamingMode )
diff --git a/Engine/Source/Editor/KismetWidgets/Public/SScrubWidget.h b/Engine/Source/Editor/KismetWidgets/Public/SScrubWidget.h
--- a/Engine/Source/Editor/KismetWidgets/Public/SScrubWidget.h
+++ b/Engine/Source/Editor/KismetWidgets/Public/SScrubWidget.h
@@ -12,6 +12,7 @@
  * Keyboard-based manipulation: click on the spinbox to enter text mode.
  */
 DECLARE_DELEGATE_TwoParams( FOnCropAnimSequence, bool, float )
+DECLARE_DELEGATE_TwoParams( FOnAddAnimSequence, bool, int32 )
 DECLARE_DELEGATE_TwoParams( FOnScrubBarDrag, int32, float)
 
 class KISMETWIDGETS_API SScrubWidget : public SCompoundWidget
@@ -49,6 +50,8 @@ public:
 		SLATE_EVENT( FOnSetInputViewRange, OnSetInputViewRange )
 		/** Called when an anim sequence is cropped before/after a selected frame */
 		SLATE_EVENT( FOnCropAnimSequence, OnCropAnimSequence )
+		/** Called when an frane is added before/after a selected frame */
+		SLATE_EVENT( FOnAddAnimSequence, OnAddAnimSequence )
 		/** Called to zero out selected frame's translation from origin */
 		SLATE_EVENT( FSimpleDelegate, OnReZeroAnimSequence )
 		/** Optional, additional values to draw on the timeline **/
@@ -93,6 +96,9 @@ private:
 	/** Function to crop animation sequence before/after selected frame */
 	void OnSequenceCropped( bool bFromStart, float CurrentFrameTime );
 
+	/** Function to crop animation sequence before/after selected frame */
+	void OnSequenceAdded(bool bBefore, int32 CurrentFrameNumber);
+
 	/** Function to zero out translation of the selected frame */
 	void OnReZero();
 
@@ -105,6 +111,7 @@ private:
 	TAttribute<float> ViewInputMax;
 	FOnSetInputViewRange OnSetInputViewRange;
 	FOnCropAnimSequence OnCropAnimSequence;
+	FOnAddAnimSequence	OnAddAnimSequence;
 	FSimpleDelegate		OnReZeroAnimSequence;
 
 	/** Dragagble bars are generic lines drawn on the scrub widget that can be dragged with the mouse. This is very bare bones and just represents drawing/moving float values */
diff --git a/Engine/Source/Editor/Persona/Private/SAnimationScrubPanel.cpp b/Engine/Source/Editor/Persona/Private/SAnimationScrubPanel.cpp
--- a/Engine/Source/Editor/Persona/Private/SAnimationScrubPanel.cpp
+++ b/Engine/Source/Editor/Persona/Private/SAnimationScrubPanel.cpp
@@ -52,6 +52,7 @@ void SAnimationScrubPanel::Construct( const SAnimationScrubPanel::FArguments& In
 			.ViewInputMax(InArgs._ViewInputMax)
 			.OnSetInputViewRange(InArgs._OnSetInputViewRange)
 			.OnCropAnimSequence( this, &SAnimationScrubPanel::OnCropAnimSequence )
+			.OnAddAnimSequence( this, &SAnimationScrubPanel::OnInsertAnimSequence )
 			.OnReZeroAnimSequence( this, &SAnimationScrubPanel::OnReZeroAnimSequence )
 			.bAllowZoom(InArgs._bAllowZoom)
 			.IsRealtimeStreamingMode(this, &SAnimationScrubPanel::IsRealtimeStreamingMode)
@@ -483,6 +484,36 @@ void SAnimationScrubPanel::OnCropAnimSequence( bool bFromStart, float CurrentTim
 	}
 }
 
+void SAnimationScrubPanel::OnInsertAnimSequence( bool bBefore, int32 CurrentFrame )
+{
+	UAnimSingleNodeInstance* PreviewInstance = GetPreviewInstance();
+	if(PreviewInstance)
+	{
+		float Length = PreviewInstance->GetLength();
+		if(PreviewInstance->CurrentAsset)
+		{
+			UAnimSequence* AnimSequence = Cast<UAnimSequence>(PreviewInstance->CurrentAsset);
+			if(AnimSequence)
+			{
+				const FScopedTransaction Transaction(LOCTEXT("InsertAnimSequence", "Insert Animation Sequence"));
+
+				//Call modify to restore slider position
+				PreviewInstance->Modify();
+
+				//Call modify to restore anim sequence current state
+				AnimSequence->Modify();
+
+				// Crop the raw anim data.
+				int32 StartFrame = (bBefore)? CurrentFrame : CurrentFrame + 1;
+				int32 EndFrame = StartFrame + 1;
+				AnimSequence->InsertFramesToRawAnimData(StartFrame, EndFrame, CurrentFrame);
+				// Recompress animation from Raw.
+				FAnimationUtils::CompressAnimSequence(AnimSequence, false, false);
+			}
+		}
+	}
+}
+
 void SAnimationScrubPanel::OnReZeroAnimSequence( )
 {
 	UAnimSingleNodeInstance* PreviewInstance = GetPreviewInstance();
diff --git a/Engine/Source/Editor/Persona/Private/SAnimationScrubPanel.h b/Engine/Source/Editor/Persona/Private/SAnimationScrubPanel.h
--- a/Engine/Source/Editor/Persona/Private/SAnimationScrubPanel.h
+++ b/Engine/Source/Editor/Persona/Private/SAnimationScrubPanel.h
@@ -64,6 +64,7 @@ protected:
 
 	/** Function to crop animation sequence before/after selected frame */
 	void OnCropAnimSequence( bool bFromStart, float CurrentTime );
+	void OnInsertAnimSequence( bool bBefore, int32 CurrentFrame );
 
 	/** Function to zero out translation of the selected frame */
 	void OnReZeroAnimSequence( );
diff --git a/Engine/Source/Runtime/Engine/Classes/Animation/AnimSequence.h b/Engine/Source/Runtime/Engine/Classes/Animation/AnimSequence.h
--- a/Engine/Source/Runtime/Engine/Classes/Animation/AnimSequence.h
+++ b/Engine/Source/Runtime/Engine/Classes/Animation/AnimSequence.h
@@ -723,6 +723,19 @@ public:
 	 */
 	ENGINE_API bool CropRawAnimData( float CurrentTime, bool bFromStart );
 
+		
+	/**
+	 * Crops the raw anim data either from Start to CurrentTime or CurrentTime to End depending on
+	 * value of bFromStart.  Can't be called against cooked data.
+	 *
+	 * @param	StartFrame		StartFrame to insert (0-based)
+	 * @param	EndFrame		EndFrame to insert (0-based
+	 * @param	CopyFrame		A frame that we copy from (0-based)
+	 * @return					true if the operation was successful.
+	 */
+	ENGINE_API bool InsertFramesToRawAnimData( int32 StartFrame, int32 EndFrame, int32 CopyFrame);
+
+
 	/** Clears any data in the AnimSequence, so it can be recycled when importing a new animation with same name over it. */
 	ENGINE_API void RecycleAnimSequence();
 
@@ -844,6 +857,9 @@ private:
 	 */
 	int32 InsertTrack(const FName& BoneName);
 
+	/** Resize Sequence Length and Num of Frames */
+	void ResizeSequence(float NewLength, int32 NewNumFrames);
+
 	friend class UAnimationAsset;
 };
 
diff --git a/Engine/Source/Runtime/Engine/Classes/Animation/AnimSequenceBase.h b/Engine/Source/Runtime/Engine/Classes/Animation/AnimSequenceBase.h
--- a/Engine/Source/Runtime/Engine/Classes/Animation/AnimSequenceBase.h
+++ b/Engine/Source/Runtime/Engine/Classes/Animation/AnimSequenceBase.h
@@ -342,6 +342,7 @@ struct FFloatCurve : public FAnimCurveBase
 	void CopyCurve(FFloatCurve& SourceCurve);
 	float Evaluate(float CurrentTime, float BlendWeight) const;
 	void UpdateOrAddKey(float NewKey, float CurrentTime);
+	void Resize(float NewLength);
 };
 
 USTRUCT()
@@ -382,6 +383,7 @@ struct FVectorCurve : public FAnimCurveBase
 	FVector Evaluate(float CurrentTime, float BlendWeight) const;
 	void UpdateOrAddKey(const FVector& NewKey, float CurrentTime);
 	bool DoesContainKey() const { return (FloatCurves[0].GetNumKeys() > 0 || FloatCurves[1].GetNumKeys() > 0 || FloatCurves[2].GetNumKeys() > 0);}
+	void Resize(float NewLength);
 };
 
 USTRUCT()
@@ -423,6 +425,7 @@ struct FTransformCurve: public FAnimCurveBase
 	void CopyCurve(FTransformCurve& SourceCurve);
 	FTransform Evaluate(float CurrentTime, float BlendWeight) const;
 	ENGINE_API void UpdateOrAddKey(const FTransform& NewKey, float CurrentTime);
+	void Resize(float NewLength);
 };
 /**
  * Raw Curve data for serialization
@@ -499,6 +502,10 @@ struct FRawCurveTracks
 	 */
 	void Serialize(FArchive& Ar);
 
+	/*
+	 * resize curve length. If longer, it doesn't do any. If shorter, remove previous keys and add new key to the end of the frame. 
+	 */
+	void Resize(float TotalLength);
 	/** 
 	 * Clear all keys
 	 */
diff --git a/Engine/Source/Runtime/Engine/Classes/Curves/CurveBase.h b/Engine/Source/Runtime/Engine/Classes/Curves/CurveBase.h
--- a/Engine/Source/Runtime/Engine/Classes/Curves/CurveBase.h
+++ b/Engine/Source/Runtime/Engine/Classes/Curves/CurveBase.h
@@ -360,6 +360,9 @@ public:
 	/** Auto set tangents for any 'auto' keys in curve */
 	void AutoSetTangents(float Tension = 0.f);
 
+	/** Resize curve length to the [MinTimeRange, MaxTimeRange] */
+	void ResizeTimeRange(float NewMinTimeRange, float NewMaxTimeRange);
+
 	/** Determine if two RichCurves are the same */
 	bool operator == (const FRichCurve& Curve) const;
 private:
diff --git a/Engine/Source/Runtime/Engine/Private/Animation/AnimSequence.cpp b/Engine/Source/Runtime/Engine/Private/Animation/AnimSequence.cpp
--- a/Engine/Source/Runtime/Engine/Private/Animation/AnimSequence.cpp
+++ b/Engine/Source/Runtime/Engine/Private/Animation/AnimSequence.cpp
@@ -1238,6 +1238,70 @@ static int32 CropRawTrack(FRawAnimSequenceTrack& RawTrack, int32 StartKey, int32
 	return FMath::Max<int32>( RawTrack.PosKeys.Num(), FMath::Max<int32>(RawTrack.RotKeys.Num(), RawTrack.ScaleKeys.Num()) );
 }
 
+void UAnimSequence::ResizeSequence(float NewLength, int32 NewNumFrames)
+{
+	NumFrames = NewNumFrames;
+	// Update sequence length to match new number of frames.
+	SequenceLength = NewLength;
+
+	// resize curves
+	RawCurveData.Resize(NewLength);
+}
+
+bool UAnimSequence::InsertFramesToRawAnimData( int32 StartFrame, int32 EndFrame, int32 CopyFrame)
+{
+	// make sure the copyframe is valid and start frame is valid
+	int32 NumFramesToInsert = EndFrame-StartFrame;
+	if ((CopyFrame>=0 && CopyFrame<NumFrames) && (StartFrame >= 0 && StartFrame <=NumFrames) && NumFramesToInsert > 0)
+	{
+		for (auto& RawData : RawAnimationData)
+		{
+			if (RawData.PosKeys.Num() > 1 && RawData.PosKeys.IsValidIndex(CopyFrame))
+			{
+				auto Source = RawData.PosKeys[CopyFrame];
+				RawData.PosKeys.InsertZeroed(StartFrame, NumFramesToInsert);
+				for (int32 Index=StartFrame; Index<EndFrame; ++Index)
+				{
+					RawData.PosKeys[Index] = Source;
+				}
+			}
+
+			if(RawData.RotKeys.Num() > 1 && RawData.RotKeys.IsValidIndex(CopyFrame))
+			{
+				auto Source = RawData.RotKeys[CopyFrame];
+				RawData.RotKeys.InsertZeroed(StartFrame, NumFramesToInsert);
+				for(int32 Index=StartFrame; Index<EndFrame; ++Index)
+				{
+					RawData.RotKeys[Index] = Source;
+				}
+			}
+
+			if(RawData.ScaleKeys.Num() > 1 && RawData.ScaleKeys.IsValidIndex(CopyFrame))
+			{
+				auto Source = RawData.ScaleKeys[CopyFrame];
+				RawData.ScaleKeys.InsertZeroed(StartFrame, NumFramesToInsert);
+
+				for(int32 Index=StartFrame; Index<EndFrame; ++Index)
+				{
+					RawData.ScaleKeys[Index] = Source;
+				}
+			}
+		}
+
+		float const FrameTime = SequenceLength / ((float)NumFrames);
+
+		int32 NewNumFrames = NumFrames + NumFramesToInsert;
+		ResizeSequence((float)NewNumFrames * FrameTime, NewNumFrames);
+
+		UE_LOG(LogAnimation, Log, TEXT("\tSequenceLength: %f, NumFrames: %d"), SequenceLength, NumFrames);
+
+		MarkPackageDirty();
+
+		return true;
+	}
+
+	return false;
+}
 
 bool UAnimSequence::CropRawAnimData( float CurrentTime, bool bFromStart )
 {
@@ -1350,7 +1414,7 @@ bool UAnimSequence::CropRawAnimData( float CurrentTime, bool bFromStart )
 	}
 
 	// Update sequence length to match new number of frames.
-	SequenceLength = (float)NumFrames * FrameTime;
+	ResizeSequence((float)NumFrames * FrameTime, NumFrames);
 
 	UE_LOG(LogAnimation, Log, TEXT("\tSequenceLength: %f, NumFrames: %d"), SequenceLength, NumFrames);
 
diff --git a/Engine/Source/Runtime/Engine/Private/Animation/AnimSequenceBase.cpp b/Engine/Source/Runtime/Engine/Private/Animation/AnimSequenceBase.cpp
--- a/Engine/Source/Runtime/Engine/Private/Animation/AnimSequenceBase.cpp
+++ b/Engine/Source/Runtime/Engine/Private/Animation/AnimSequenceBase.cpp
@@ -139,6 +139,10 @@ void FFloatCurve::UpdateOrAddKey(float NewKey, float CurrentTime)
 	FloatCurve.UpdateOrAddKey(CurrentTime, NewKey);
 }
 
+void FFloatCurve::Resize(float NewLength)
+{
+	FloatCurve.ResizeTimeRange(0, NewLength);
+}
 ////////////////////////////////////////////////////
 //  FVectorCurve
 
@@ -165,6 +169,13 @@ void FVectorCurve::UpdateOrAddKey(const FVector& NewKey, float CurrentTime)
 	FloatCurves[Y].UpdateOrAddKey(CurrentTime, NewKey.Y);
 	FloatCurves[Z].UpdateOrAddKey(CurrentTime, NewKey.Z);
 }
+
+void FVectorCurve::Resize(float NewLength)
+{
+	FloatCurves[X].ResizeTimeRange(0, NewLength);
+	FloatCurves[Y].ResizeTimeRange(0, NewLength);
+	FloatCurves[Z].ResizeTimeRange(0, NewLength);
+}
 ////////////////////////////////////////////////////
 //  FTransformCurve
 
@@ -211,6 +222,14 @@ void FTransformCurve::UpdateOrAddKey(const FTransform& NewKey, float CurrentTime
 	RotationCurve.UpdateOrAddKey(RotationAsVector, CurrentTime);
 	ScaleCurve.UpdateOrAddKey(NewKey.GetScale3D(), CurrentTime);
 }
+
+void FTransformCurve::Resize(float NewLength)
+{
+	TranslationCurve.Resize(NewLength);
+	RotationCurve.Resize(NewLength);
+	ScaleCurve.Resize(NewLength);
+}
+
 /////////////////////////////////////////////////////
 // FRawCurveTracks
 
@@ -307,6 +326,24 @@ bool FRawCurveTracks::AddCurveData(USkeleton::AnimCurveUID Uid, int32 CurveFlags
 	}
 }
 
+void FRawCurveTracks::Resize(float TotalLength)
+{
+	for (auto& Curve: FloatCurves)
+	{
+		Curve.Resize(TotalLength);
+	}
+
+	for(auto& Curve: VectorCurves)
+	{
+		Curve.Resize(TotalLength);
+	}
+
+	for(auto& Curve: TransformCurves)
+	{
+		Curve.Resize(TotalLength);
+	}
+}
+
 void FRawCurveTracks::Serialize(FArchive& Ar)
 {
 	// @TODO: If we're about to serialize vector curve, add here
diff --git a/Engine/Source/Runtime/Engine/Private/CurveBase.cpp b/Engine/Source/Runtime/Engine/Private/CurveBase.cpp
--- a/Engine/Source/Runtime/Engine/Private/CurveBase.cpp
+++ b/Engine/Source/Runtime/Engine/Private/CurveBase.cpp
@@ -664,6 +664,49 @@ void FRichCurve::AutoSetTangents(float Tension)
 	}
 }
 
+void FRichCurve::ResizeTimeRange(float NewMinTimeRange, float NewMaxTimeRange)
+{
+	float MinTime, MaxTime;
+	GetTimeRange(MinTime, MaxTime);
+
+	bool bNeedToDeleteKey=false;
+
+	// if there is key below min time, just add key at new min range, 
+	if (MinTime < NewMinTimeRange)
+	{
+		float NewValue = Eval(NewMinTimeRange);
+		AddKey(NewMinTimeRange, NewValue);
+
+		bNeedToDeleteKey = true;
+	}
+
+	// if there is key after max time, just add key at new max range, 
+	if(MaxTime > NewMaxTimeRange)
+	{
+		float NewValue = Eval(NewMaxTimeRange);
+		AddKey(NewMaxTimeRange, NewValue);
+
+		bNeedToDeleteKey = true;
+	}
+
+	// delete the keys outside of range
+	if (bNeedToDeleteKey)
+	{
+		for (int32 KeyIndex=0; KeyIndex<Keys.Num(); ++KeyIndex)
+		{
+			if (Keys[KeyIndex].Time < NewMinTimeRange || Keys[KeyIndex].Time > NewMaxTimeRange)
+			{
+				const FKeyHandle* KeyHandle = KeyHandlesToIndices.FindKey(KeyIndex);
+				if (KeyHandle)
+				{
+					DeleteKey(*KeyHandle);
+					--KeyIndex;
+				}
+			}
+		}
+	}
+}
+
 /** Util to find float value on bezier defined by 4 control points */ 
 static float BezierInterp(float P0, float P1, float P2, float P3, float Alpha)
 {
