Improved the Pin Type Selector filter so order of words and case sensitivity does not matter.

diff --git a/Engine/Source/Editor/KismetWidgets/Private/SPinTypeSelector.cpp b/Engine/Source/Editor/KismetWidgets/Private/SPinTypeSelector.cpp
--- a/Engine/Source/Editor/KismetWidgets/Private/SPinTypeSelector.cpp
+++ b/Engine/Source/Editor/KismetWidgets/Private/SPinTypeSelector.cpp
@@ -343,7 +343,7 @@ void SPinTypeSelector::OnFilterTextChanged(const FText& NewText)
 	SearchText = NewText;
 	FilteredTypeTreeRoot.Empty();
 
-	GetChildrenMatchingSearch(FName::NameToDisplayString(NewText.ToString(), false), TypeTreeRoot, FilteredTypeTreeRoot);
+	GetChildrenMatchingSearch(NewText, TypeTreeRoot, FilteredTypeTreeRoot);
 	TypeTreeView->RequestTreeRefresh();
 
 	// Select the first non-category item
@@ -374,8 +374,26 @@ void SPinTypeSelector::OnFilterTextCommitted(const FText& NewText, ETextCommit::
 	}
 }
 
-bool SPinTypeSelector::GetChildrenMatchingSearch(const FString& InSearchText, const TArray<FPinTypeTreeItem>& UnfilteredList, TArray<FPinTypeTreeItem>& OutFilteredList)
+bool SPinTypeSelector::GetChildrenMatchingSearch(const FText& InSearchText, const TArray<FPinTypeTreeItem>& UnfilteredList, TArray<FPinTypeTreeItem>& OutFilteredList)
 {
+	// Trim and sanitized the filter text (so that it more likely matches the action descriptions)
+	FString TrimmedFilterString = FText::TrimPrecedingAndTrailing(InSearchText).ToString();
+
+	// Tokenize the search box text into a set of terms; all of them must be present to pass the filter
+	TArray<FString> FilterTerms;
+	TrimmedFilterString.ParseIntoArray(FilterTerms, TEXT(" "), true);
+
+	// Generate a list of sanitized versions of the strings
+	TArray<FString> SanitizedFilterTerms;
+	for (int32 iFilters = 0; iFilters < FilterTerms.Num() ; iFilters++)
+	{
+		FString EachString = FName::NameToDisplayString( FilterTerms[iFilters], false );
+		EachString = EachString.Replace( TEXT( " " ), TEXT( "" ) );
+		SanitizedFilterTerms.Add( EachString );
+	}
+	// Both of these should match!
+	ensure( SanitizedFilterTerms.Num() == FilterTerms.Num() );
+
 	bool bReturnVal = false;
 
 	for( auto it = UnfilteredList.CreateConstIterator(); it; ++it )
@@ -384,10 +402,24 @@ bool SPinTypeSelector::GetChildrenMatchingSearch(const FString& InSearchText, co
 		FPinTypeTreeItem NewInfo = MakeShareable( new UEdGraphSchema_K2::FPinTypeTreeInfo(Item) );
 		TArray<FPinTypeTreeItem> ValidChildren;
 
-		// Have to run GetChildrenMatchingSearch first, so that we can make sure we get valid children for the list!
-		if( GetChildrenMatchingSearch(InSearchText, Item->Children, ValidChildren)
-			|| InSearchText.IsEmpty()
-			|| Item->GetDescription().ToString().Contains(InSearchText) )
+		const bool bHasChildrenMatchingSearch = GetChildrenMatchingSearch(InSearchText, Item->Children, ValidChildren);
+		const bool bIsEmptySearch = InSearchText.IsEmpty();
+		bool bFilterTextMatches = true;
+
+		// If children match the search filter or it's an empty search, let's not do any checks against the FilterTerms
+		if ( !bHasChildrenMatchingSearch && !bIsEmptySearch)
+		{
+			FString DescriptionString =  Item->GetDescription().ToString();
+			DescriptionString = DescriptionString.Replace( TEXT( " " ), TEXT( "" ) );
+			for (int32 FilterIndex = 0; (FilterIndex < FilterTerms.Num()) && bFilterTextMatches; ++FilterIndex)
+			{
+				const bool bMatchesTerm = ( DescriptionString.Contains( FilterTerms[FilterIndex] ) || ( DescriptionString.Contains( SanitizedFilterTerms[FilterIndex] ) == true ) );
+				bFilterTextMatches = bFilterTextMatches && bMatchesTerm;
+			}
+		}
+		if( bHasChildrenMatchingSearch
+			|| bIsEmptySearch
+			|| bFilterTextMatches )
 		{
 			NewInfo->Children = ValidChildren;
 			OutFilteredList.Add(NewInfo);
diff --git a/Engine/Source/Editor/KismetWidgets/Public/SPinTypeSelector.h b/Engine/Source/Editor/KismetWidgets/Public/SPinTypeSelector.h
--- a/Engine/Source/Editor/KismetWidgets/Public/SPinTypeSelector.h
+++ b/Engine/Source/Editor/KismetWidgets/Public/SPinTypeSelector.h
@@ -123,7 +123,7 @@ protected:
 	void OnFilterTextCommitted(const FText& NewText, ETextCommit::Type CommitInfo);
 
 	/** Helper to generate the filtered list of types, based on the search string matching */
-	bool GetChildrenMatchingSearch(const FString& SearchText, const TArray<FPinTypeTreeItem>& UnfilteredList, TArray<FPinTypeTreeItem>& OutFilteredList);
+	bool GetChildrenMatchingSearch(const FText& SearchText, const TArray<FPinTypeTreeItem>& UnfilteredList, TArray<FPinTypeTreeItem>& OutFilteredList);
 
 	/** 
 	 * Determine the best icon to represent the given pin.
