Tooltips will display for sub-menu items in the Pin Type Selector.

diff --git a/Engine/Source/Editor/KismetWidgets/Private/SPinTypeSelector.cpp b/Engine/Source/Editor/KismetWidgets/Private/SPinTypeSelector.cpp
--- a/Engine/Source/Editor/KismetWidgets/Private/SPinTypeSelector.cpp
+++ b/Engine/Source/Editor/KismetWidgets/Private/SPinTypeSelector.cpp
@@ -278,7 +278,7 @@ TSharedRef<ITableRow> SPinTypeSelector::GenerateTypeTreeRow(FPinTypeTreeItem InI
 	return ReturnWidget;
 }
 
-void SPinTypeSelector::AddObjectReferenceMenuEntry(FMenuBuilder& InOutMenuBuilder, FPinTypeTreeItem InItem, FEdGraphPinType& InPinType, const FSlateBrush* InIconBrush) const
+void SPinTypeSelector::AddObjectReferenceMenuEntry(FMenuBuilder& InOutMenuBuilder, FPinTypeTreeItem InItem, FEdGraphPinType& InPinType, const FSlateBrush* InIconBrush, FText InSimpleTooltip) const
 {
 	InOutMenuBuilder.AddMenuEntry(
 		FUIAction(
@@ -286,22 +286,23 @@ void SPinTypeSelector::AddObjectReferenceMenuEntry(FMenuBuilder& InOutMenuBuilde
 		),
 		// Contents
 		SNew(SHorizontalBox)
-			+SHorizontalBox::Slot()
+			.ToolTip(IDocumentation::Get()->CreateToolTip(InSimpleTooltip, NULL, *BigTooltipDocLink, InPinType.PinCategory))
+		+SHorizontalBox::Slot()
 			.AutoWidth()
 			.Padding(1.f)
-			[
-				SNew(SImage)
+		[
+			SNew(SImage)
 				.Image(InIconBrush)
 				.ColorAndOpacity(Schema->GetPinTypeColor(InPinType))
-			]
-			+SHorizontalBox::Slot()
+		]
+		+SHorizontalBox::Slot()
 			.AutoWidth()
 			.Padding(1.f)
-			[
-				SNew(STextBlock)
+		[
+			SNew(STextBlock)
 				.Text(UEdGraphSchema_K2::GetCategoryText(InPinType.PinCategory))
 				.Font(FEditorStyle::GetFontStyle(TEXT("Kismet.TypePicker.NormalFont")) )
-			]
+		]
 		);
 }
 
@@ -313,6 +314,13 @@ TSharedRef< SWidget > SPinTypeSelector::GetAllowedObjectTypes(FPinTypeTreeItem I
 	FEdGraphPinType PinType = InItem->GetPinType(false);
 	const FSlateBrush* IconBrush = GetIconFromPin(PinType);
 
+	FFormatNamedArguments Args;
+
+	if(PinType.PinSubCategoryObject.IsValid())
+	{
+		Args.Add(TEXT("TypeName"), InItem->GetDescription());
+	}
+
 	uint8 PossibleObjectReferenceTypes = InItem->GetPossibleObjectReferenceTypes();
 
 	// Per each object reference type, change the category to the type and add a menu entry (this will get the color to be correct)
@@ -320,25 +328,25 @@ TSharedRef< SWidget > SPinTypeSelector::GetAllowedObjectTypes(FPinTypeTreeItem I
 	if (PossibleObjectReferenceTypes & static_cast<uint8>(EObjectReferenceType::ObjectReference))
 	{
 		PinType.PinCategory = UEdGraphSchema_K2::PC_Object;
-		AddObjectReferenceMenuEntry(MenuBuilder, InItem, PinType, IconBrush);
+		AddObjectReferenceMenuEntry(MenuBuilder, InItem, PinType, IconBrush, FText::Format(LOCTEXT("ObjectTooltip", "Reference an instanced object of type \'{TypeName}\'"), Args));
 	}
 
 	if (PossibleObjectReferenceTypes & static_cast<uint8>(EObjectReferenceType::ClassReference))
 	{
 		PinType.PinCategory = UEdGraphSchema_K2::PC_Class;
-		AddObjectReferenceMenuEntry(MenuBuilder, InItem, PinType, IconBrush);
+		AddObjectReferenceMenuEntry(MenuBuilder, InItem, PinType, IconBrush, FText::Format(LOCTEXT("ClassTooltip", "Reference a class of type \'{TypeName}\'"), Args));
 	}
 
 	if (PossibleObjectReferenceTypes & static_cast<uint8>(EObjectReferenceType::AssetID))
 	{
 		PinType.PinCategory = UEdGraphSchema_K2::PC_Asset;
-		AddObjectReferenceMenuEntry(MenuBuilder, InItem, PinType, IconBrush);
+		AddObjectReferenceMenuEntry(MenuBuilder, InItem, PinType, IconBrush, FText::Format(LOCTEXT("AssetTooltip", "Path to an instanced object of type \'{Typename}\' which may be in an unloaded state. Can be utilized to asynchronously load the object reference."), Args));
 	}
 
 	if (PossibleObjectReferenceTypes & static_cast<uint8>(EObjectReferenceType::ClassAssetID))
 	{
 		PinType.PinCategory = UEdGraphSchema_K2::PC_AssetClass;
-		AddObjectReferenceMenuEntry(MenuBuilder, InItem, PinType, IconBrush);
+		AddObjectReferenceMenuEntry(MenuBuilder, InItem, PinType, IconBrush, FText::Format(LOCTEXT("ClassTooltip", "Path to a class object of type \'{Typename}\' which may be in an unloaded state. Can be utilized to asynchronously load the class."), Args));
 	}
 
 	return MenuBuilder.MakeWidget();
diff --git a/Engine/Source/Editor/KismetWidgets/Public/SPinTypeSelector.h b/Engine/Source/Editor/KismetWidgets/Public/SPinTypeSelector.h
--- a/Engine/Source/Editor/KismetWidgets/Public/SPinTypeSelector.h
+++ b/Engine/Source/Editor/KismetWidgets/Public/SPinTypeSelector.h
@@ -146,8 +146,9 @@ protected:
 	 * @param InItem				Tree item to use for the callback when a menu item is selected
 	 * @param InPinType				Pin type for generation of the widget to display for the menu entry
 	 * @param InIconBrush			Brush icon to use for the menu entry item
+	 * @param InTooltip				The simple tooltip to use for the menu item, an advanced tooltip link will be auto-generated based on the PinCategory
 	 */
-	void AddObjectReferenceMenuEntry(FMenuBuilder& InOutMenuBuilder, FPinTypeTreeItem InItem, FEdGraphPinType& InPinType, const FSlateBrush* InIconBrush) const;
+	void AddObjectReferenceMenuEntry(FMenuBuilder& InOutMenuBuilder, FPinTypeTreeItem InItem, FEdGraphPinType& InPinType, const FSlateBrush* InIconBrush, FText InSimpleTooltip) const;
 
 	/** Gets the allowable object types for an tree item, used for building the sub-menu */
 	TSharedRef< SWidget > GetAllowedObjectTypes(FPinTypeTreeItem InItem) const;
