PinTypeSelector has improved usability with the sub-menu for object types.

SListViewSelectorDropdownMenu no longer requires any SSearchBox be supplied as a DefaultFocusWidget, it will safely ignore the concept.

PinTypeSelector now uses a ListView for the sub-menu when picking "Reference", "Class", "Asset ID", or "Class Asset ID".

Issues with being able to hit enter twice to select "Reference" have been resolved.

Various usability improvements with the sub-menu.

SSubMenuHandler::RequestSubMenuToggle now supports for immediate reaction.

Added SSubMenuHandler::IsSubMenuOpen

diff --git a/Engine/Source/Editor/KismetWidgets/Private/SPinTypeSelector.cpp b/Engine/Source/Editor/KismetWidgets/Private/SPinTypeSelector.cpp
--- a/Engine/Source/Editor/KismetWidgets/Private/SPinTypeSelector.cpp
+++ b/Engine/Source/Editor/KismetWidgets/Private/SPinTypeSelector.cpp
@@ -12,6 +12,25 @@
 
 static const FString BigTooltipDocLink = TEXT("Shared/Editor/Blueprint/VariableTypes");
 
+/** Manages items in the Object Reference Type list, the sub-menu of the PinTypeSelector */
+struct FObjectReferenceType
+{
+	/** Item that is being referenced */
+	FPinTypeTreeItem PinTypeItem;
+
+	/** Widget to display for this item */
+	TSharedPtr< SWidget > WidgetToDisplay;
+
+	/** Category that should be used when this item is selected */
+	FString PinCategory;
+
+	FObjectReferenceType(FPinTypeTreeItem InPinTypeItem, TSharedRef< SWidget > InWidget, FString InPinCategory)
+		: PinTypeItem(InPinTypeItem)
+		, WidgetToDisplay(InWidget)
+		, PinCategory(InPinCategory)
+	{}
+};
+
 class SPinTypeRow : public SComboRow<FPinTypeTreeItem>
 {
 public:
@@ -43,16 +62,22 @@ public:
 	}
 	// End of SWidget interface
 
-	/** Returns TRUE if there is a Sub-Menu to open */
+	/** Returns TRUE if there is a Sub-Menu available to open */
 	bool HasSubMenu() const
 	{
 		return SubMenuHandler.Pin()->HasSubMenu();
 	}
 
+	/** Returns TRUE if there is a Sub-Menu open */
+	bool IsSubMenuOpen() const
+	{
+		return SubMenuHandler.Pin()->IsSubMenuOpen();
+	}
+
 	/** Forces the sub-menu open, clobbering any other open ones in the process */
-	void RequestSubMenuToggle()
+	void RequestSubMenuToggle(bool bInImmediate = false)
 	{
-		SubMenuHandler.Pin()->RequestSubMenuToggle(true, true);
+		SubMenuHandler.Pin()->RequestSubMenuToggle(true, true, bInImmediate);
 	}
 
 private:
@@ -81,7 +106,7 @@ void SPinTypeSelector::Construct(const FArguments& InArgs, FGetPinTypeTree GetPi
 
 	bIsRightMousePressed = false;
 
-	// Depending on if this is a compact selector or not, we generate a different compoound widget
+	// Depending on if this is a compact selector or not, we generate a different compound widget
 	TSharedPtr<SWidget> Widget;
 
 	if (InArgs._bCompactSelector)
@@ -278,37 +303,81 @@ TSharedRef<ITableRow> SPinTypeSelector::GenerateTypeTreeRow(FPinTypeTreeItem InI
 	return ReturnWidget;
 }
 
-void SPinTypeSelector::AddObjectReferenceMenuEntry(FMenuBuilder& InOutMenuBuilder, FPinTypeTreeItem InItem, FEdGraphPinType& InPinType, const FSlateBrush* InIconBrush, FText InSimpleTooltip) const
+TSharedRef<SWidget> SPinTypeSelector::CreateObjectReferenceWidget(FPinTypeTreeItem InItem, FEdGraphPinType& InPinType, const FSlateBrush* InIconBrush, FText InSimpleTooltip) const
 {
-	InOutMenuBuilder.AddMenuEntry(
-		FUIAction(
-		FExecuteAction::CreateSP(this, &SPinTypeSelector::OnSelectPinType, InItem, InPinType.PinCategory)
-		),
-		// Contents
-		SNew(SHorizontalBox)
-			.ToolTip(IDocumentation::Get()->CreateToolTip(InSimpleTooltip, NULL, *BigTooltipDocLink, InPinType.PinCategory))
+	return SNew(SHorizontalBox)
+		.ToolTip(IDocumentation::Get()->CreateToolTip(InSimpleTooltip, NULL, *BigTooltipDocLink, InPinType.PinCategory))
 		+SHorizontalBox::Slot()
-			.AutoWidth()
-			.Padding(1.f)
+		.AutoWidth()
+		.Padding(1.f)
 		[
 			SNew(SImage)
-				.Image(InIconBrush)
-				.ColorAndOpacity(Schema->GetPinTypeColor(InPinType))
+			.Image(InIconBrush)
+			.ColorAndOpacity(Schema->GetPinTypeColor(InPinType))
 		]
 		+SHorizontalBox::Slot()
-			.AutoWidth()
-			.Padding(1.f)
+		.AutoWidth()
+		.Padding(1.f)
 		[
 			SNew(STextBlock)
-				.Text(UEdGraphSchema_K2::GetCategoryText(InPinType.PinCategory))
-				.Font(FEditorStyle::GetFontStyle(TEXT("Kismet.TypePicker.NormalFont")) )
-		]
-		);
+			.Text(UEdGraphSchema_K2::GetCategoryText(InPinType.PinCategory))
+			.Font(FEditorStyle::GetFontStyle(TEXT("Kismet.TypePicker.NormalFont")) )
+		];
+}
+
+class SObjectReferenceWidget : public SCompoundWidget
+{
+public:
+	SLATE_BEGIN_ARGS( SObjectReferenceWidget )
+	{}
+		SLATE_DEFAULT_SLOT( FArguments, Content )
+	SLATE_END_ARGS()
+
+	void Construct(const FArguments& InArgs, TWeakPtr< SMenuOwner > InMenuOwner)
+	{
+		MenuOwner = InMenuOwner;
+
+		this->ChildSlot
+		[
+			InArgs._Content.Widget
+		];
+	}
+
+	// SWidget interface
+	virtual FReply OnKeyDown(const FGeometry& MyGeometry, const FKeyEvent& KeyEvent) override
+	{
+		if (MenuOwner.IsValid() && (KeyEvent.GetKey() == EKeys::Left || KeyEvent.GetKey() == EKeys::Escape))
+		{
+			MenuOwner.Pin()->CloseSummonedMenus();
+			return FReply::Handled();
+		}
+		return FReply::Unhandled();
+	}
+	// End of SWidget interface
+
+private:
+	TWeakPtr< SMenuOwner > MenuOwner;
+};
+
+TSharedRef<ITableRow> SPinTypeSelector::GenerateObjectReferenceTreeRow(FObjectReferenceListItem InItem, const TSharedRef<STableViewBase>& OwnerTree)
+{
+	return SNew(SComboRow<FObjectReferenceListItem>, OwnerTree)
+		[
+			InItem->WidgetToDisplay.ToSharedRef()
+		];
 }
 
-TSharedRef< SWidget > SPinTypeSelector::GetAllowedObjectTypes(FPinTypeTreeItem InItem) const
+void SPinTypeSelector::OnObjectReferenceSelectionChanged(FObjectReferenceListItem InItem, ESelectInfo::Type SelectInfo)
 {
-	FMenuBuilder MenuBuilder( false, MakeShareable(new FUICommandList) );
+	if (SelectInfo != ESelectInfo::OnNavigation)
+	{
+		OnSelectPinType(InItem->PinTypeItem, InItem->PinCategory);
+	}
+}
+
+TSharedRef< SWidget > SPinTypeSelector::GetAllowedObjectTypes(FPinTypeTreeItem InItem)
+{
+	AllowedObjectReferenceTypes.Reset();
 
 	// Do not force the pin type here, that causes a load of the Blueprint (if unloaded)
 	FEdGraphPinType PinType = InItem->GetPinType(false);
@@ -328,28 +397,60 @@ TSharedRef< SWidget > SPinTypeSelector::GetAllowedObjectTypes(FPinTypeTreeItem I
 	if (PossibleObjectReferenceTypes & static_cast<uint8>(EObjectReferenceType::ObjectReference))
 	{
 		PinType.PinCategory = UEdGraphSchema_K2::PC_Object;
-		AddObjectReferenceMenuEntry(MenuBuilder, InItem, PinType, IconBrush, FText::Format(LOCTEXT("ObjectTooltip", "Reference an instanced object of type \'{TypeName}\'"), Args));
+		TSharedRef<SWidget> Widget = CreateObjectReferenceWidget(InItem, PinType, IconBrush, FText::Format(LOCTEXT("ObjectTooltip", "Reference an instanced object of type \'{TypeName}\'"), Args));
+		FObjectReferenceListItem ObjectReferenceType = MakeShareable(new FObjectReferenceType(InItem, Widget, PinType.PinCategory));
+		AllowedObjectReferenceTypes.Add(ObjectReferenceType);
 	}
 
 	if (PossibleObjectReferenceTypes & static_cast<uint8>(EObjectReferenceType::ClassReference))
 	{
 		PinType.PinCategory = UEdGraphSchema_K2::PC_Class;
-		AddObjectReferenceMenuEntry(MenuBuilder, InItem, PinType, IconBrush, FText::Format(LOCTEXT("ClassTooltip", "Reference a class of type \'{TypeName}\'"), Args));
+		TSharedRef<SWidget> Widget = CreateObjectReferenceWidget(InItem, PinType, IconBrush, FText::Format(LOCTEXT("ClassTooltip", "Reference a class of type \'{TypeName}\'"), Args));
+		FObjectReferenceListItem ObjectReferenceType = MakeShareable(new FObjectReferenceType(InItem, Widget, PinType.PinCategory));
+		AllowedObjectReferenceTypes.Add(ObjectReferenceType);
 	}
 
 	if (PossibleObjectReferenceTypes & static_cast<uint8>(EObjectReferenceType::AssetID))
 	{
 		PinType.PinCategory = UEdGraphSchema_K2::PC_Asset;
-		AddObjectReferenceMenuEntry(MenuBuilder, InItem, PinType, IconBrush, FText::Format(LOCTEXT("AssetTooltip", "Path to an instanced object of type \'{Typename}\' which may be in an unloaded state. Can be utilized to asynchronously load the object reference."), Args));
+		TSharedRef<SWidget> Widget = CreateObjectReferenceWidget(InItem, PinType, IconBrush, FText::Format(LOCTEXT("AssetTooltip", "Path to an instanced object of type \'{Typename}\' which may be in an unloaded state. Can be utilized to asynchronously load the object reference."), Args));
+		FObjectReferenceListItem ObjectReferenceType = MakeShareable(new FObjectReferenceType(InItem, Widget, PinType.PinCategory));
+		AllowedObjectReferenceTypes.Add(ObjectReferenceType);
 	}
 
 	if (PossibleObjectReferenceTypes & static_cast<uint8>(EObjectReferenceType::ClassAssetID))
 	{
 		PinType.PinCategory = UEdGraphSchema_K2::PC_AssetClass;
-		AddObjectReferenceMenuEntry(MenuBuilder, InItem, PinType, IconBrush, FText::Format(LOCTEXT("ClassTooltip", "Path to a class object of type \'{Typename}\' which may be in an unloaded state. Can be utilized to asynchronously load the class."), Args));
+		TSharedRef<SWidget> Widget = CreateObjectReferenceWidget(InItem, PinType, IconBrush, FText::Format(LOCTEXT("ClassTooltip", "Path to a class object of type \'{Typename}\' which may be in an unloaded state. Can be utilized to asynchronously load the class."), Args));
+		FObjectReferenceListItem ObjectReferenceType = MakeShareable(new FObjectReferenceType(InItem, Widget, PinType.PinCategory));
+		AllowedObjectReferenceTypes.Add(ObjectReferenceType);
 	}
 
-	return MenuBuilder.MakeWidget();
+	TSharedPtr<SListView<FObjectReferenceListItem>> ListView;
+	SAssignNew(ListView, SListView<FObjectReferenceListItem>)
+		.ListItemsSource(&AllowedObjectReferenceTypes)
+		.SelectionMode(ESelectionMode::Single)
+		.OnGenerateRow(this, &SPinTypeSelector::GenerateObjectReferenceTreeRow)
+		.OnSelectionChanged(this, &SPinTypeSelector::OnObjectReferenceSelectionChanged);
+
+	WeakListView = ListView;
+	if (AllowedObjectReferenceTypes.Num())
+	{
+		ListView->SetSelection(AllowedObjectReferenceTypes[0], ESelectInfo::OnNavigation);
+	}
+
+	return 
+		SNew(SBorder)
+		.BorderImage(FEditorStyle::GetBrush("Menu.Background"))
+		[
+			SNew(SObjectReferenceWidget, PinTypeSelectorMenuOwner)
+			[
+				SNew(SListViewSelectorDropdownMenu<FObjectReferenceListItem>, nullptr, ListView)
+				[
+					ListView.ToSharedRef()
+				]
+			]
+		];
 }
 
 void SPinTypeSelector::OnSelectPinType(FPinTypeTreeItem InItem, FString InPinCategory)
@@ -382,6 +483,12 @@ void SPinTypeSelector::OnTypeSelectionChanged(FPinTypeTreeItem Selection, ESelec
 	// When the user is navigating, do not act upon the selection change
 	if(SelectInfo == ESelectInfo::OnNavigation)
 	{
+		// Unless mouse clicking on an item with a sub-menu, all attempts to auto-select should open the sub-menu
+		TSharedPtr<SPinTypeRow> PinRow = StaticCastSharedPtr<SPinTypeRow>(TypeTreeView->WidgetFromItem(Selection));
+		if (PinRow.IsValid() && PinTypeSelectorMenuOwner.IsValid())
+		{
+			PinTypeSelectorMenuOwner.Pin()->CloseSummonedMenus();
+		}
 		return;
 	}
 
@@ -392,9 +499,10 @@ void SPinTypeSelector::OnTypeSelectionChanged(FPinTypeTreeItem Selection, ESelec
 		{
 			// Unless mouse clicking on an item with a sub-menu, all attempts to auto-select should open the sub-menu
 			TSharedPtr<SPinTypeRow> PinRow = StaticCastSharedPtr<SPinTypeRow>(TypeTreeView->WidgetFromItem(Selection));
-			if (SelectInfo != ESelectInfo::OnMouseClick && PinRow.IsValid() && PinRow->HasSubMenu())
+			if (SelectInfo != ESelectInfo::OnMouseClick && PinRow.IsValid() && PinRow->HasSubMenu() && !PinRow->IsSubMenuOpen())
 			{
-				PinRow->RequestSubMenuToggle();
+				PinRow->RequestSubMenuToggle(true);
+				FSlateApplication::Get().SetKeyboardFocus(WeakListView.Pin(), EFocusCause::SetDirectly);
 			}
 			else
 			{
@@ -460,7 +568,7 @@ TSharedRef<SWidget>	SPinTypeSelector::GetMenuContent()
 			.OnTextChanged( this, &SPinTypeSelector::OnFilterTextChanged )
 			.OnTextCommitted( this, &SPinTypeSelector::OnFilterTextCommitted );
 
-		MenuContent = SNew(SMenuOwner)
+		MenuContent = SAssignNew(PinTypeSelectorMenuOwner, SMenuOwner)
 			[
 				SNew(SListViewSelectorDropdownMenu<FPinTypeTreeItem>, FilterTextBox, TypeTreeView)
 				[
diff --git a/Engine/Source/Editor/KismetWidgets/Public/SPinTypeSelector.h b/Engine/Source/Editor/KismetWidgets/Public/SPinTypeSelector.h
--- a/Engine/Source/Editor/KismetWidgets/Public/SPinTypeSelector.h
+++ b/Engine/Source/Editor/KismetWidgets/Public/SPinTypeSelector.h
@@ -15,6 +15,9 @@ typedef STreeView<FPinTypeTreeItem> SPinTypeTreeView;
 
 DECLARE_DELEGATE_ThreeParams(FGetPinTypeTree, TArray<FPinTypeTreeItem >&, bool, bool);
 
+struct FObjectReferenceType;
+typedef TSharedPtr<struct FObjectReferenceType> FObjectReferenceListItem;
+
 /** Widget for modifying the type for a variable or pin */
 class KISMETWIDGETS_API SPinTypeSelector : public SCompoundWidget
 {
@@ -96,6 +99,11 @@ protected:
 	/** TRUE if displaying a compact selector widget, some functionality is enabled in different ways if this is TRUE */
 	bool bIsCompactSelector;
 
+	/** Holds a cache of the allowed Object Reference types for the last sub-menu opened. */
+	TArray<FObjectReferenceListItem> AllowedObjectReferenceTypes;
+	TWeakPtr<SListView<FObjectReferenceListItem>> WeakListView;
+	TWeakPtr<SMenuOwner> PinTypeSelectorMenuOwner;
+
 	/** Array checkbox support functions */
 	ECheckBoxState IsArrayChecked() const;
 	void OnArrayCheckStateChanged(ECheckBoxState NewState);
@@ -113,6 +121,10 @@ protected:
 	void OnTypeSelectionChanged(FPinTypeTreeItem Selection, ESelectInfo::Type SelectInfo);
 	void GetTypeChildren(FPinTypeTreeItem InItem, TArray<FPinTypeTreeItem>& OutChildren);
 
+	/** Listview support functions for sub-menu */
+	TSharedRef<ITableRow> GenerateObjectReferenceTreeRow(FObjectReferenceListItem InItem, const TSharedRef<STableViewBase>& OwnerTree);
+	void OnObjectReferenceSelectionChanged(FObjectReferenceListItem InItem, ESelectInfo::Type SelectInfo);
+
 	/** Reference to the menu content that's displayed when the type button is clicked on */
 	TSharedPtr<SMenuOwner> MenuContent;
 	virtual TSharedRef<SWidget>	GetMenuContent();
@@ -140,18 +152,17 @@ protected:
 	FText GetToolTipForArrayWidget() const;
 
 	/**
-	 * Helper function to add an object reference menu entry to a MenuBuilder
+	 * Helper function to create widget for the sub-menu
 	 *
-	 * @param InOutMenuBuilder		Menu builder to add the menu entry to
 	 * @param InItem				Tree item to use for the callback when a menu item is selected
 	 * @param InPinType				Pin type for generation of the widget to display for the menu entry
 	 * @param InIconBrush			Brush icon to use for the menu entry item
 	 * @param InTooltip				The simple tooltip to use for the menu item, an advanced tooltip link will be auto-generated based on the PinCategory
 	 */
-	void AddObjectReferenceMenuEntry(FMenuBuilder& InOutMenuBuilder, FPinTypeTreeItem InItem, FEdGraphPinType& InPinType, const FSlateBrush* InIconBrush, FText InSimpleTooltip) const;
+	TSharedRef<SWidget> CreateObjectReferenceWidget(FPinTypeTreeItem InItem, FEdGraphPinType& InPinType, const FSlateBrush* InIconBrush, FText InSimpleTooltip) const;
 
 	/** Gets the allowable object types for an tree item, used for building the sub-menu */
-	TSharedRef< SWidget > GetAllowedObjectTypes(FPinTypeTreeItem InItem) const;
+	TSharedRef< SWidget > GetAllowedObjectTypes(FPinTypeTreeItem InItem);
 	
 	/**
 	 * When a pin type is selected, handle it
diff --git a/Engine/Source/Editor/UnrealEd/Public/SListViewSelectorDropdownMenu.h b/Engine/Source/Editor/UnrealEd/Public/SListViewSelectorDropdownMenu.h
--- a/Engine/Source/Editor/UnrealEd/Public/SListViewSelectorDropdownMenu.h
+++ b/Engine/Source/Editor/UnrealEd/Public/SListViewSelectorDropdownMenu.h
@@ -21,14 +21,16 @@ public:
 	 */
 	void Construct(const FArguments& InArgs, TSharedPtr<SSearchBox> InDefaultFocusWidget, TSharedPtr< SListView< ListType > > InTargetListView)
 	{
-		check(InDefaultFocusWidget.IsValid());
 		check(InTargetListView.IsValid());
 
 		TargetListView = InTargetListView;
 		DefaultFocusWidget = InDefaultFocusWidget;
 
-		// We have some overrides for OnKeyDown handling of the SSearchBox to provide the seamless functionality
-		InDefaultFocusWidget->SetOnKeyDownHandler(FOnKeyDown::CreateSP(this, &SListViewSelectorDropdownMenu<ListType>::OnKeyDown));
+		if (DefaultFocusWidget.IsValid())
+		{
+			// We have some overrides for OnKeyDown handling of the SSearchBox to provide the seamless functionality
+			InDefaultFocusWidget->SetOnKeyDownHandler(FOnKeyDown::CreateSP(this, &SListViewSelectorDropdownMenu<ListType>::OnKeyDown));
+		}
 
 		this->ChildSlot
 		[
@@ -55,7 +57,7 @@ public:
 			}
 			return FReply::Handled();
 		}
-		else
+		else if (DefaultFocusWidget.IsValid())
 		{
 			TSharedPtr< SWidget > DefaultFocusWidgetPtr = DefaultFocusWidget.Pin();
 			FReply Reply = DefaultFocusWidgetPtr->OnKeyDown(FindChildGeometry(MyGeometry, TargetListViewPtr.ToSharedRef()), KeyEvent);
diff --git a/Engine/Source/Runtime/Slate/Private/Widgets/Input/SSubMenuHandler.cpp b/Engine/Source/Runtime/Slate/Private/Widgets/Input/SSubMenuHandler.cpp
--- a/Engine/Source/Runtime/Slate/Private/Widgets/Input/SSubMenuHandler.cpp
+++ b/Engine/Source/Runtime/Slate/Private/Widgets/Input/SSubMenuHandler.cpp
@@ -95,15 +95,22 @@ bool SSubMenuHandler::ShouldSubMenuAppearHovered() const
 	return MenuAnchor.IsValid() && MenuAnchor.Pin()->IsOpen() && !MenuOwnerWidget.Pin()->IsHovered();
 }
 
-void SSubMenuHandler::RequestSubMenuToggle( bool bOpenMenu, const bool bClobber )
+void SSubMenuHandler::RequestSubMenuToggle( bool bOpenMenu, const bool bClobber, const bool bImmediate )
 {
 	if (MenuAnchor.IsValid())
 	{
-		// Reset the time before the menu opens
-		float TimeToSubMenuOpen = bClobber ? MultiBoxConstants::SubMenuClobberTime : MultiBoxConstants::SubMenuOpenTime;
-		if (!ActiveTimerHandle.IsValid())
+		if (bImmediate)
 		{
-			ActiveTimerHandle = RegisterActiveTimer(TimeToSubMenuOpen, FWidgetActiveTimerDelegate::CreateSP(this, &SSubMenuHandler::UpdateSubMenuState, bOpenMenu));
+			UpdateSubMenuState(0, 0, true);
+		}
+		else
+		{
+			// Reset the time before the menu opens
+			float TimeToSubMenuOpen = bClobber ? MultiBoxConstants::SubMenuClobberTime : MultiBoxConstants::SubMenuOpenTime;
+			if (!ActiveTimerHandle.IsValid())
+			{
+				ActiveTimerHandle = RegisterActiveTimer(TimeToSubMenuOpen, FWidgetActiveTimerDelegate::CreateSP(this, &SSubMenuHandler::UpdateSubMenuState, bOpenMenu));
+			}
 		}
 	}
 }
@@ -120,6 +127,11 @@ void SSubMenuHandler::CancelPendingSubMenu()
 	}
 }
 
+bool SSubMenuHandler::IsSubMenuOpen() const
+{
+	return MenuAnchor.IsValid() && MenuAnchor.Pin()->IsOpen();
+}
+
 EActiveTimerReturnType SSubMenuHandler::UpdateSubMenuState(double InCurrentTime, float InDeltaTime, bool bWantsOpen)
 {
 	TSharedPtr< SMenuOwner > PinnedMenuOwnerWidget(MenuOwnerWidget.Pin());
diff --git a/Engine/Source/Runtime/Slate/Public/Widgets/Input/SSubMenuHandler.h b/Engine/Source/Runtime/Slate/Public/Widgets/Input/SSubMenuHandler.h
--- a/Engine/Source/Runtime/Slate/Public/Widgets/Input/SSubMenuHandler.h
+++ b/Engine/Source/Runtime/Slate/Public/Widgets/Input/SSubMenuHandler.h
@@ -34,10 +34,11 @@ public:
 	 * Requests that the sub-menu associated with this widget be toggled on or off.
 	 * It does not immediately toggle the menu.  After a set amount of time is passed the menu will toggle
 	 *
-	 * @param bOpenMenu	true to open the menu, false to close the menu if one is currently open
-	 * @param bClobber true if we want to open a menu when another menu is already open
+	 * @param bOpenMenu					true to open the menu, false to close the menu if one is currently open
+	 * @param bClobber					true if we want to open a menu when another menu is already open
+	 * @param const bool bImmediate		Default false, when true the menu will open immediately
 	 */
-	void RequestSubMenuToggle( bool bOpenMenu, const bool bClobber );
+	void RequestSubMenuToggle( bool bOpenMenu, const bool bClobber, const bool bImmediate = false);
 	
 	/**
 	 * Cancels any open requests to toggle a sub-menu       
@@ -50,6 +51,9 @@ public:
 		return MenuAnchor.IsValid();
 	}
 
+	/** Returns TRUE if the sub-menu is open */
+	bool IsSubMenuOpen() const;
+
 protected:
 	/** One-off delayed active timer to update the open/closed state of the sub menu. */
 	EActiveTimerReturnType UpdateSubMenuState(double InCurrentTime, float InDeltaTime, bool bWantsOpen);
