Copying //UE4/Dev-Blueprints to //UE4/Dev-Main (Source: //UE4/Dev-Blueprints @ 3130440)

diff --git a/Engine/Source/Editor/KismetWidgets/Private/SPinTypeSelector.cpp b/Engine/Source/Editor/KismetWidgets/Private/SPinTypeSelector.cpp
--- a/Engine/Source/Editor/KismetWidgets/Private/SPinTypeSelector.cpp
+++ b/Engine/Source/Editor/KismetWidgets/Private/SPinTypeSelector.cpp
@@ -11,6 +11,43 @@
 
 #define LOCTEXT_NAMESPACE "PinTypeSelector"
 
+/** I need a widget that draws two images on top of each other. This is to represent a TMap (key type and value type): */
+class SDoubleImage : public SImage
+{
+public:
+	void Construct(const FArguments& InArgs, TAttribute<const FSlateBrush*> InSecondImage, TAttribute<FSlateColor> InSecondImageColor);
+
+private:
+	virtual int32 OnPaint(const FPaintArgs& Args, const FGeometry& AllottedGeometry, const FSlateRect& MyClippingRect, FSlateWindowElementList& OutDrawElements, int32 LayerId, const FWidgetStyle& InWidgetStyle, bool bParentEnabled) const override;
+
+	TAttribute<const FSlateBrush*> SecondImage;
+	TAttribute<FSlateColor> SecondImageColor;
+};
+
+void SDoubleImage::Construct(const SDoubleImage::FArguments& InArgs, TAttribute<const FSlateBrush*> InSecondImage, TAttribute<FSlateColor> InSecondImageColor)
+{
+	SImage::Construct(InArgs);
+	SecondImage = InSecondImage;
+	SecondImageColor = InSecondImageColor;
+}
+
+int32 SDoubleImage::OnPaint(const FPaintArgs& Args, const FGeometry& AllottedGeometry, const FSlateRect& MyClippingRect, FSlateWindowElementList& OutDrawElements, int32 LayerId, const FWidgetStyle& InWidgetStyle, bool bParentEnabled) const
+{
+	// this will draw Image[0]:
+	SImage::OnPaint(Args, AllottedGeometry, MyClippingRect, OutDrawElements, LayerId, InWidgetStyle, bParentEnabled);
+
+	const bool bIsEnabled = ShouldBeEnabled(bParentEnabled);
+	const uint32 DrawEffects = bIsEnabled ? ESlateDrawEffect::None : ESlateDrawEffect::DisabledEffect;
+	// draw rest of the images, we reuse the LayerId because images are assumed to note overlap:
+	const FSlateBrush* SecondImageResolved = SecondImage.Get();
+	if (SecondImageResolved && SecondImageResolved->DrawAs != ESlateBrushDrawType::NoDrawType)
+	{
+		const FLinearColor FinalColorAndOpacity(InWidgetStyle.GetColorAndOpacityTint() * SecondImageColor.Get().GetColor(InWidgetStyle) * SecondImageResolved->GetTint(InWidgetStyle));
+		FSlateDrawElement::MakeBox(OutDrawElements, LayerId, AllottedGeometry.ToPaintGeometry(), SecondImageResolved, MyClippingRect, DrawEffects, FinalColorAndOpacity);
+	}
+	return LayerId;
+}
+
 static const FString BigTooltipDocLink = TEXT("Shared/Editor/Blueprint/VariableTypes");
 
 /** Manages items in the Object Reference Type list, the sub-menu of the PinTypeSelector */
@@ -86,8 +123,35 @@ private:
 	TWeakPtr<SSubMenuHandler> SubMenuHandler;
 };
 
+enum class EPinContainerType
+{
+	SingleValue,
+	Array,
+	Set,
+	Map,
+};
+
+TSharedRef<SWidget> SPinTypeSelector::ConstructPinTypeImage(const FSlateBrush* PrimaryIcon, const FSlateColor& PrimaryColor, const FSlateBrush* SecondaryIcon, const FSlateColor& SecondaryColor, TSharedPtr<SToolTip> InToolTip)
+{
+	return
+		SNew(SDoubleImage, SecondaryIcon, SecondaryColor)
+		.Image(PrimaryIcon)
+		.ToolTip(InToolTip)
+		.ColorAndOpacity(PrimaryColor);
+}
+
 void SPinTypeSelector::Construct(const FArguments& InArgs, FGetPinTypeTree GetPinTypeTreeFunc)
 {
+	// SComboBox is a bit restrictive:
+	static TArray<TSharedPtr<EPinContainerType>> PinTypes;
+	if (PinTypes.Num() == 0)
+	{
+		PinTypes.Add(MakeShareable(new EPinContainerType(EPinContainerType::SingleValue)));
+		PinTypes.Add(MakeShareable(new EPinContainerType(EPinContainerType::Array)));
+		PinTypes.Add(MakeShareable(new EPinContainerType(EPinContainerType::Set)));
+		PinTypes.Add(MakeShareable(new EPinContainerType(EPinContainerType::Map)));
+	}
+
 	SearchText = FText::GetEmpty();
 
 	OnTypeChanged = InArgs._OnPinTypeChanged;
@@ -113,13 +177,22 @@ void SPinTypeSelector::Construct(const FArguments& InArgs, FGetPinTypeTree GetPi
 	{
 		// Only have a combo button with an icon
 		Widget = SAssignNew( TypeComboButton, SComboButton )
-			.OnGetMenuContent(this, &SPinTypeSelector::GetMenuContent)
+			.OnGetMenuContent(this, &SPinTypeSelector::GetMenuContent, false)
 			.ContentPadding(0)
 			.ToolTipText(this, &SPinTypeSelector::GetToolTipForComboBoxType)
 			.HasDownArrow(!InArgs._bCompactSelector)
 			.ButtonStyle(FEditorStyle::Get(),  "BlueprintEditor.CompactPinTypeSelector")
 			.ButtonContent()
 			[
+				FBlueprintEditorUtils::ShouldEnableAdvancedContainers() ?
+				SNew(
+					SDoubleImage,
+					TAttribute<const FSlateBrush*>(this, &SPinTypeSelector::GetSecondaryTypeIconImage),
+					TAttribute<FSlateColor>(this, &SPinTypeSelector::GetSecondaryTypeIconColor)
+				)
+				.Image(this, &SPinTypeSelector::GetTypeIconImage)
+				.ColorAndOpacity(this, &SPinTypeSelector::GetTypeIconColor)
+				:
 				SNew(SImage)
 				.Image( this, &SPinTypeSelector::GetTypeIconImage )
 				.ColorAndOpacity( this, &SPinTypeSelector::GetTypeIconColor )
@@ -128,54 +201,176 @@ void SPinTypeSelector::Construct(const FArguments& InArgs, FGetPinTypeTree GetPi
 	else
 	{
 		// Traditional Pin Type Selector with a combo button, the icon, the current type name, and a toggle button for being an array
-		Widget = SNew(SHorizontalBox)
-		+SHorizontalBox::Slot()
-		[
-			SAssignNew( TypeComboButton, SComboButton )
-			.OnGetMenuContent(this, &SPinTypeSelector::GetMenuContent)
+		TSharedPtr<SWidget> ContainerControl;
+		if (FBlueprintEditorUtils::ShouldEnableAdvancedContainers())
+		{
+			auto GenerateContainerTypeEntry = [this](TSharedPtr<EPinContainerType> InPinContainerType)
+			{
+				EPinContainerType PinContainerType = *InPinContainerType;
+
+				static const FSlateBrush* Images[] = {
+					FEditorStyle::GetBrush(TEXT("Kismet.VariableList.TypeIcon")),
+					FEditorStyle::GetBrush(TEXT("Kismet.VariableList.ArrayTypeIcon")),
+					FEditorStyle::GetBrush(TEXT("Kismet.VariableList.SetTypeIcon")),
+					FEditorStyle::GetBrush(TEXT("Kismet.VariableList.MapKeyTypeIcon")),
+				};
+				check(sizeof(Images) / sizeof(*Images) > (int32)PinContainerType);
+				const FSlateBrush* SecondaryIcon = PinContainerType == EPinContainerType::Map ? FEditorStyle::GetBrush(TEXT("Kismet.VariableList.MapValueTypeIcon")) : nullptr;
+
+				static const FText Tooltips[] = {
+					LOCTEXT("SingleVariableTooltip", "Single Variable"),
+					LOCTEXT("ArrayTooltip", "Array"),
+					LOCTEXT("SetTooltip", "Set"),
+					LOCTEXT("MapTooltip", "Map (Dictionary)"),
+				};
+				check(sizeof(Tooltips) / sizeof(*Tooltips) > (int32)PinContainerType);
+
+
+				return SNew(
+							SDoubleImage,
+							SecondaryIcon,
+							TAttribute<FSlateColor>(this, &SPinTypeSelector::GetSecondaryTypeIconColor)
+						)
+					.Image(Images[(int32)PinContainerType])
+					.ToolTip(IDocumentation::Get()->CreateToolTip(Tooltips[(int32)PinContainerType], nullptr, *BigTooltipDocLink, TEXT("Containers")))
+					.ColorAndOpacity(this, &SPinTypeSelector::GetTypeIconColor);
+			};
+
+			ContainerControl = SNew(SComboBox< TSharedPtr<EPinContainerType> >)
+			.ButtonStyle(FCoreStyle::Get(), "NoBorder")
+			.HasDownArrow(false)
+			.OnGenerateWidget(
+				SComboBox< TSharedPtr<EPinContainerType> >::FOnGenerateWidget::CreateLambda(
+					GenerateContainerTypeEntry
+				)
+			)
+			.OptionsSource(&PinTypes)
 			.ContentPadding(0)
-			.ToolTipText(this, &SPinTypeSelector::GetToolTipForComboBoxType)
-			.ButtonContent()
+			.ToolTip(IDocumentation::Get()->CreateToolTip( TAttribute<FText>(this, &SPinTypeSelector::GetToolTipForContainerWidget), NULL, *BigTooltipDocLink, TEXT("Containers")))
+			.IsEnabled( TargetPinType.Get().PinCategory != Schema->PC_Exec )
+			.Visibility(InArgs._bAllowArrays ? EVisibility::Visible : EVisibility::Collapsed)
+			.OnSelectionChanged(
+				SComboBox< TSharedPtr<EPinContainerType> >::FOnSelectionChanged::CreateLambda(
+					[this](TSharedPtr<EPinContainerType> InType, ESelectInfo::Type)
+					{
+						this->OnContainerTypeSelectionChanged(*InType);
+					}
+				)
+			)
+			.Content()
 			[
-				SNew(SHorizontalBox)
-				+SHorizontalBox::Slot()
-				.AutoWidth()
-				.VAlign(VAlign_Center)
-				.HAlign(HAlign_Center)
+				SNew(
+					SDoubleImage,
+					TAttribute<const FSlateBrush*>(this, &SPinTypeSelector::GetSecondaryTypeIconImage),
+					TAttribute<FSlateColor>(this, &SPinTypeSelector::GetSecondaryTypeIconColor)
+				)
+				.Image(this, &SPinTypeSelector::GetTypeIconImage)
+				.ColorAndOpacity(this, &SPinTypeSelector::GetTypeIconColor)
+			];
+		}
+		else
+		{
+			ContainerControl = SNew( SCheckBox )
+				.ToolTip(IDocumentation::Get()->CreateToolTip( TAttribute<FText>(this, &SPinTypeSelector::GetToolTipForArrayWidget), NULL, *BigTooltipDocLink, TEXT("Array")))
+				.Padding( 1 )
+				.OnCheckStateChanged( this, &SPinTypeSelector::OnArrayCheckStateChanged )
+				.IsChecked( this, &SPinTypeSelector::IsArrayChecked )
+				.IsEnabled( TargetPinType.Get().PinCategory != Schema->PC_Exec )
+				.Style( FEditorStyle::Get(), "ToggleButtonCheckbox" )
+				.Visibility(InArgs._bAllowArrays ? EVisibility::Visible : EVisibility::Collapsed)
 				[
 					SNew(SImage)
-					.Image( this, &SPinTypeSelector::GetTypeIconImage )
+					.Image( FEditorStyle::GetBrush(TEXT("Kismet.VariableList.ArrayTypeIcon")) )
 					.ColorAndOpacity( this, &SPinTypeSelector::GetTypeIconColor )
+				];
+		}
+
+		Widget = SNew(SHorizontalBox)
+		+SHorizontalBox::Slot()
+		[
+			SNew(SBox)
+			.WidthOverride(100.f)
+			[
+					SAssignNew( TypeComboButton, SComboButton )
+					.OnGetMenuContent(this, &SPinTypeSelector::GetMenuContent, false)
+					.ContentPadding(0)
+					.ToolTipText(this, &SPinTypeSelector::GetToolTipForComboBoxType)
+					.ButtonContent()
+					[
+						SNew(SHorizontalBox)
+						+SHorizontalBox::Slot()
+						.VAlign(VAlign_Center)
+						.HAlign(HAlign_Left)
+						.AutoWidth()
+						[
+							SNew(SImage)
+							.Image( this, &SPinTypeSelector::GetTypeIconImage )
+							.ColorAndOpacity( this, &SPinTypeSelector::GetTypeIconColor )
+						]
+						+SHorizontalBox::Slot()
+						.VAlign(VAlign_Center)
+						.HAlign(HAlign_Left)
+						.AutoWidth()
+						[
+							SNew(STextBlock)
+							.Text( this, &SPinTypeSelector::GetTypeDescription )
+							.Font(InArgs._Font)
+						]
+						+ SHorizontalBox::Slot()
+						.AutoWidth()
+						[
+							SNew(SSpacer)
+						]
+					]
 				]
-				+SHorizontalBox::Slot()
-				.AutoWidth()
-				.VAlign(VAlign_Center)
-				.HAlign(HAlign_Left)
-				[
-					SNew(STextBlock)
-					.Text( this, &SPinTypeSelector::GetTypeDescription )
-					.Font(InArgs._Font)
-				]
-			]
 		]
 
 		+SHorizontalBox::Slot()
 		.AutoWidth()
-		.VAlign(VAlign_Center)
-		.HAlign(HAlign_Center)
+			.VAlign(VAlign_Center)
+			.HAlign(HAlign_Center)
 		[
-			SNew( SCheckBox )
-			.ToolTip(IDocumentation::Get()->CreateToolTip( TAttribute<FText>(this, &SPinTypeSelector::GetToolTipForArrayWidget), NULL, *BigTooltipDocLink, TEXT("Array")))
-			.Padding( 1 )
-			.OnCheckStateChanged( this, &SPinTypeSelector::OnArrayCheckStateChanged )
-			.IsChecked( this, &SPinTypeSelector::IsArrayChecked )
-			.IsEnabled( TargetPinType.Get().PinCategory != Schema->PC_Exec )
-			.Style( FEditorStyle::Get(), "ToggleButtonCheckbox" )
-			.Visibility(InArgs._bAllowArrays ? EVisibility::Visible : EVisibility::Collapsed)
+			ContainerControl.ToSharedRef()
+		]
+
+		+SHorizontalBox::Slot()
+		[
+			SNew(SBox)
+			.WidthOverride(110.f)
+			.Visibility(
+				TAttribute<EVisibility>::Create(
+					TAttribute<EVisibility>::FGetter::CreateLambda(
+						[this]() {return this->TargetPinType.Get().bIsMap == true ? EVisibility::Visible : EVisibility::Hidden; }
+					)
+				)
+			)
 			[
-				SNew(SImage)
-				.Image( FEditorStyle::GetBrush(TEXT("Kismet.VariableList.ArrayTypeIcon")) )
-				.ColorAndOpacity( this, &SPinTypeSelector::GetTypeIconColor )
+				SNew( SComboButton )
+				.OnGetMenuContent(this, &SPinTypeSelector::GetMenuContent, true )
+				.ContentPadding(0)
+				.ToolTipText(this, &SPinTypeSelector::GetToolTipForComboBoxSecondaryType)
+				.ButtonContent()
+				[
+					SNew(SHorizontalBox)
+					+SHorizontalBox::Slot()
+					.AutoWidth()
+					.VAlign(VAlign_Center)
+					.HAlign(HAlign_Center)
+					[
+						SNew(SImage)
+						.Image( this, &SPinTypeSelector::GetSecondaryTypeIconImage )
+						.ColorAndOpacity( this, &SPinTypeSelector::GetSecondaryTypeIconColor )
+					]
+					+SHorizontalBox::Slot()
+					.AutoWidth()
+					.VAlign(VAlign_Center)
+					.HAlign(HAlign_Left)
+					[
+						SNew(STextBlock)
+						.Text( this, &SPinTypeSelector::GetSecondaryTypeDescription )
+						.Font(InArgs._Font)
+					]
+				]
 			]
 		];
 	}
@@ -206,16 +401,45 @@ FText SPinTypeSelector::GetTypeDescription() const
 	}
 }
 
+FText SPinTypeSelector::GetSecondaryTypeDescription() const
+{
+	const FString& PinSubCategory = TargetPinType.Get().PinValueType.TerminalSubCategory;
+	const UObject* PinSubCategoryObject = TargetPinType.Get().PinValueType.TerminalSubCategoryObject.Get();
+	if (PinSubCategory != UEdGraphSchema_K2::PSC_Bitmask && PinSubCategoryObject)
+	{
+		if (auto Field = Cast<const UField>(PinSubCategoryObject))
+		{
+			return Field->GetDisplayNameText();
+		}
+		return FText::FromString(PinSubCategoryObject->GetName());
+	}
+	else
+	{
+		return UEdGraphSchema_K2::GetCategoryText(TargetPinType.Get().PinValueType.TerminalCategory, true);
+	}
+}
+
+
 const FSlateBrush* SPinTypeSelector::GetTypeIconImage() const
 {
 	return FBlueprintEditorUtils::GetIconFromPin( TargetPinType.Get() );
 }
 
+const FSlateBrush* SPinTypeSelector::GetSecondaryTypeIconImage() const
+{
+	return FBlueprintEditorUtils::GetSecondaryIconFromPin(TargetPinType.Get());
+}
+
 FSlateColor SPinTypeSelector::GetTypeIconColor() const
 {
 	return Schema->GetPinTypeColor(TargetPinType.Get());
 }
 
+FSlateColor SPinTypeSelector::GetSecondaryTypeIconColor() const
+{
+	return Schema->GetSecondaryPinTypeColor(TargetPinType.Get());
+}
+
 ECheckBoxState SPinTypeSelector::IsArrayChecked() const
 {
 	return TargetPinType.Get().bIsArray ? ECheckBoxState::Checked : ECheckBoxState::Unchecked;
@@ -224,7 +448,10 @@ ECheckBoxState SPinTypeSelector::IsArrayChecked() const
 void SPinTypeSelector::OnArrayCheckStateChanged(ECheckBoxState NewState)
 {
 	FEdGraphPinType NewTargetPinType = TargetPinType.Get();
-	NewTargetPinType.bIsArray = (NewState == ECheckBoxState::Checked)? true : false;
+	NewTargetPinType.bIsArray = (NewState == ECheckBoxState::Checked) ? true : false;
+	// just in case someone has been experimenting with experimental features:
+	NewTargetPinType.bIsMap = false;
+	NewTargetPinType.bIsSet = false;
 
 	OnTypeChanged.ExecuteIfBound(NewTargetPinType);
 }
@@ -234,9 +461,40 @@ void SPinTypeSelector::OnArrayStateToggled()
 	OnArrayCheckStateChanged((IsArrayChecked() == ECheckBoxState::Checked)? ECheckBoxState::Unchecked : ECheckBoxState::Checked);
 }
 
+void SPinTypeSelector::OnContainerTypeSelectionChanged( EPinContainerType PinContainerType)
+{
+	FEdGraphPinType NewTargetPinType = TargetPinType.Get();
+
+	switch (PinContainerType)
+	{
+	case EPinContainerType::SingleValue:
+		NewTargetPinType.bIsArray = false;
+		NewTargetPinType.bIsMap = false;
+		NewTargetPinType.bIsSet = false;
+		break;
+	case EPinContainerType::Array:
+		NewTargetPinType.bIsArray = true;
+		NewTargetPinType.bIsMap = false;
+		NewTargetPinType.bIsSet = false;
+		break;
+	case EPinContainerType::Set:
+		NewTargetPinType.bIsArray = false;
+		NewTargetPinType.bIsMap = false;
+		NewTargetPinType.bIsSet = true;
+		break;
+	case EPinContainerType::Map:
+		NewTargetPinType.bIsArray = false;
+		NewTargetPinType.bIsMap = true;
+		NewTargetPinType.bIsSet = false;
+		break;
+	}
+
+	OnTypeChanged.ExecuteIfBound(NewTargetPinType);
+}
+
 //=======================================================================
 // Type TreeView Support
-TSharedRef<ITableRow> SPinTypeSelector::GenerateTypeTreeRow(FPinTypeTreeItem InItem, const TSharedRef<STableViewBase>& OwnerTree)
+TSharedRef<ITableRow> SPinTypeSelector::GenerateTypeTreeRow(FPinTypeTreeItem InItem, const TSharedRef<STableViewBase>& OwnerTree, bool bForSecondaryType)
 {
 	const bool bHasChildren = (InItem->Children.Num() > 0);
 	const FText Description = InItem->GetDescription();
@@ -255,7 +513,7 @@ TSharedRef<ITableRow> SPinTypeSelector::GenerateTypeTreeRow(FPinTypeTreeItem InI
 	FOnGetContent OnGetContent;
 	if (InItem->GetPossibleObjectReferenceTypes() != static_cast<uint8>(EObjectReferenceType::NotAnObject))
 	{
-		OnGetContent.BindSP(this, &SPinTypeSelector::GetAllowedObjectTypes, InItem);
+		OnGetContent.BindSP(this, &SPinTypeSelector::GetAllowedObjectTypes, InItem, bForSecondaryType);
 	}
 
 	TSharedPtr< SHorizontalBox > HorizontalBox;
@@ -368,15 +626,15 @@ TSharedRef<ITableRow> SPinTypeSelector::GenerateObjectReferenceTreeRow(FObjectRe
 		];
 }
 
-void SPinTypeSelector::OnObjectReferenceSelectionChanged(FObjectReferenceListItem InItem, ESelectInfo::Type SelectInfo)
+void SPinTypeSelector::OnObjectReferenceSelectionChanged(FObjectReferenceListItem InItem, ESelectInfo::Type SelectInfo, bool bForSecondaryType)
 {
 	if (SelectInfo != ESelectInfo::OnNavigation)
 	{
-		OnSelectPinType(InItem->PinTypeItem, InItem->PinCategory);
+		OnSelectPinType(InItem->PinTypeItem, InItem->PinCategory, bForSecondaryType);
 	}
 }
 
-TSharedRef< SWidget > SPinTypeSelector::GetAllowedObjectTypes(FPinTypeTreeItem InItem)
+TSharedRef< SWidget > SPinTypeSelector::GetAllowedObjectTypes(FPinTypeTreeItem InItem, bool bForSecondaryType)
 {
 	AllowedObjectReferenceTypes.Reset();
 
@@ -432,7 +690,7 @@ TSharedRef< SWidget > SPinTypeSelector::GetAllowedObjectTypes(FPinTypeTreeItem I
 		.ListItemsSource(&AllowedObjectReferenceTypes)
 		.SelectionMode(ESelectionMode::Single)
 		.OnGenerateRow(this, &SPinTypeSelector::GenerateObjectReferenceTreeRow)
-		.OnSelectionChanged(this, &SPinTypeSelector::OnObjectReferenceSelectionChanged);
+		.OnSelectionChanged(this, &SPinTypeSelector::OnObjectReferenceSelectionChanged, bForSecondaryType);
 
 	WeakListView = ListView;
 	if (AllowedObjectReferenceTypes.Num())
@@ -454,7 +712,7 @@ TSharedRef< SWidget > SPinTypeSelector::GetAllowedObjectTypes(FPinTypeTreeItem I
 		];
 }
 
-void SPinTypeSelector::OnSelectPinType(FPinTypeTreeItem InItem, FString InPinCategory)
+void SPinTypeSelector::OnSelectPinType(FPinTypeTreeItem InItem, FString InPinCategory, bool bForSecondaryType)
 {
 	const FScopedTransaction Transaction( LOCTEXT("ChangeParam", "Change Parameter Type") );
 
@@ -465,21 +723,36 @@ void SPinTypeSelector::OnSelectPinType(FPinTypeTreeItem InItem, FString InPinCat
 	const FEdGraphPinType& SelectionPinType = InItem->GetPinType(true);
 
 	// Change the pin's type
-	NewTargetPinType.PinCategory = InPinCategory;
-	NewTargetPinType.PinSubCategory = SelectionPinType.PinSubCategory;
-	NewTargetPinType.PinSubCategoryObject = SelectionPinType.PinSubCategoryObject;
+	if (bForSecondaryType)
+	{
+		NewTargetPinType.PinValueType.TerminalCategory = InPinCategory;
+		NewTargetPinType.PinValueType.TerminalSubCategory = SelectionPinType.PinSubCategory;
+		NewTargetPinType.PinValueType.TerminalSubCategoryObject = SelectionPinType.PinSubCategoryObject;
+	}
+	else
+	{
+		NewTargetPinType.PinCategory = InPinCategory;
+		NewTargetPinType.PinSubCategory = SelectionPinType.PinSubCategory;
+		NewTargetPinType.PinSubCategoryObject = SelectionPinType.PinSubCategoryObject;
+	}
+
 
 	TypeComboButton->SetIsOpen(false);
 
 	if( NewTargetPinType.PinCategory == Schema->PC_Exec )
 	{
+		NewTargetPinType.bIsMap = false;
+		NewTargetPinType.bIsSet = false;
 		NewTargetPinType.bIsArray = false;
+		NewTargetPinType.PinValueType.TerminalCategory = FString();
+		NewTargetPinType.PinValueType.TerminalSubCategory = FString();
+		NewTargetPinType.PinValueType.TerminalSubCategoryObject = nullptr;
 	}
 
 	OnTypeChanged.ExecuteIfBound(NewTargetPinType);
 }
 
-void SPinTypeSelector::OnTypeSelectionChanged(FPinTypeTreeItem Selection, ESelectInfo::Type SelectInfo)
+void SPinTypeSelector::OnTypeSelectionChanged(FPinTypeTreeItem Selection, ESelectInfo::Type SelectInfo, bool bForSecondaryType)
 {
 	// When the user is navigating, do not act upon the selection change
 	if(SelectInfo == ESelectInfo::OnNavigation)
@@ -507,7 +780,7 @@ void SPinTypeSelector::OnTypeSelectionChanged(FPinTypeTreeItem Selection, ESelec
 			}
 			else
 			{
-				OnSelectPinType(Selection, Selection->GetPossibleObjectReferenceTypes() == static_cast<uint8>(EObjectReferenceType::AllTypes)? Schema->PC_Object : Selection->GetPinType(false).PinCategory);
+				OnSelectPinType(Selection, Selection->GetPossibleObjectReferenceTypes() == static_cast<uint8>(EObjectReferenceType::AllTypes)? Schema->PC_Object : Selection->GetPinType(false).PinCategory, bForSecondaryType);
 			}
 		}
 		else
@@ -532,7 +805,7 @@ void SPinTypeSelector::GetTypeChildren(FPinTypeTreeItem InItem, TArray<FPinTypeT
 	OutChildren = InItem->Children;
 }
 
-TSharedRef<SWidget>	SPinTypeSelector::GetMenuContent()
+TSharedRef<SWidget>	SPinTypeSelector::GetMenuContent(bool bForSecondaryType)
 {
 	GetPinTypeTree.Execute(TypeTreeRoot, TypeTreeFilter);
 
@@ -561,8 +834,8 @@ TSharedRef<SWidget>	SPinTypeSelector::GetMenuContent()
 		SAssignNew(TypeTreeView, SPinTypeTreeView)
 			.TreeItemsSource(&FilteredTypeTreeRoot)
 			.SelectionMode(ESelectionMode::Single)
-			.OnGenerateRow(this, &SPinTypeSelector::GenerateTypeTreeRow)
-			.OnSelectionChanged(this, &SPinTypeSelector::OnTypeSelectionChanged)
+			.OnGenerateRow( this, &SPinTypeSelector::GenerateTypeTreeRow, bForSecondaryType )
+			.OnSelectionChanged( this, &SPinTypeSelector::OnTypeSelectionChanged, bForSecondaryType )
 			.OnGetChildren(this, &SPinTypeSelector::GetTypeChildren);
 
 		SAssignNew(FilterTextBox, SSearchBox)
@@ -727,6 +1000,16 @@ FText SPinTypeSelector::GetToolTipForComboBoxType() const
 	return LOCTEXT("PinTypeSelector_Disabled", "Cannot edit variable type when they are inherited from parent.");
 }
 
+FText SPinTypeSelector::GetToolTipForComboBoxSecondaryType() const
+{
+	if (IsEnabled())
+	{
+		return LOCTEXT("PinTypeSelector", "Select the map's value type.");
+	}
+
+	return LOCTEXT("PinTypeSelector_Disabled", "Cannot edit map value type when they are inherited from parent.");
+}
+
 FText SPinTypeSelector::GetToolTipForArrayWidget() const
 {
 	if(IsEnabled())
@@ -742,6 +1025,21 @@ FText SPinTypeSelector::GetToolTipForArrayWidget() const
 	return LOCTEXT("ArrayCheckBox_Disabled", "Cannot edit variable type while the variable is placed in a graph or inherited from parent.");
 }
 
+FText SPinTypeSelector::GetToolTipForContainerWidget() const
+{
+	if (IsEnabled())
+	{
+		// The entire widget may be enabled, but the container type button may be disabled because it is an "exec" pin.
+		if (TargetPinType.Get().PinCategory == Schema->PC_Exec)
+		{
+			return LOCTEXT("ContainerType_ExecDisabled", "Exec pins cannot be containers.");
+		}
+		return LOCTEXT("ContainerType", "Make this variable a container (array, set, or map) of selected type.");
+	}
+
+	return LOCTEXT("ContainerType_Disabled", "Cannot edit variable type while the variable is placed in a graph or inherited from parent.");
+}
+
 FReply SPinTypeSelector::OnMouseButtonDown( const FGeometry& MyGeometry, const FPointerEvent& MouseEvent )
 {
 	if(bIsCompactSelector && MouseEvent.GetEffectingButton() == EKeys::RightMouseButton)
diff --git a/Engine/Source/Editor/KismetWidgets/Public/SPinTypeSelector.h b/Engine/Source/Editor/KismetWidgets/Public/SPinTypeSelector.h
--- a/Engine/Source/Editor/KismetWidgets/Public/SPinTypeSelector.h
+++ b/Engine/Source/Editor/KismetWidgets/Public/SPinTypeSelector.h
@@ -22,6 +22,8 @@ typedef TSharedPtr<struct FObjectReferenceType> FObjectReferenceListItem;
 class KISMETWIDGETS_API SPinTypeSelector : public SCompoundWidget
 {
 public:
+	static TSharedRef<SWidget> ConstructPinTypeImage(const FSlateBrush* PrimaryIcon, const FSlateColor& PrimaryColor, const FSlateBrush* SecondaryIcon, const FSlateColor& SecondaryColor, TSharedPtr<SToolTip> InToolTip);
+
 	SLATE_BEGIN_ARGS( SPinTypeSelector )
 		: _TargetPinType()
 		, _Schema(NULL)
@@ -53,14 +55,23 @@ public:
 	// End of SWidget interface
 
 protected:
-	/** Gets the icon (pin or array) for the type being manipulated */
+	/** Gets the icon (value, array, set, or map) for the type being manipulated */
 	const FSlateBrush* GetTypeIconImage() const;
 
+	/** Gets the secondary icon (for maps, otherwise null) for the type being manipulated */
+	const FSlateBrush* GetSecondaryTypeIconImage() const;
+
 	/** Gets the type-specific color for the type being manipulated */
 	FSlateColor GetTypeIconColor() const;
 
+	/** Gets the secondary type-specific color for the type being manipulated */
+	FSlateColor GetSecondaryTypeIconColor() const;
+
 	/** Gets a succinct type description for the type being manipulated */
-	virtual FText GetTypeDescription() const;
+	FText GetTypeDescription() const;
+
+	/** Gets the secondary type description. E.g. the value type for TMaps */
+	FText GetSecondaryTypeDescription() const;
 
 	TSharedPtr<SComboButton>		TypeComboButton;
 	TSharedPtr<SSearchBox>			FilterTextBox;
@@ -106,23 +117,26 @@ protected:
 	/** Toggles the variable type as an array */
 	void OnArrayStateToggled();
 
+	/** Updates the variable container type: */
+	void OnContainerTypeSelectionChanged(enum class EPinContainerType PinContainerType);
+
 	/** Array containing the unfiltered list of all supported types this pin could possibly have */
 	TArray<FPinTypeTreeItem>		TypeTreeRoot;
 	/** Array containing a filtered list, according to the text in the searchbox */
 	TArray<FPinTypeTreeItem>		FilteredTypeTreeRoot;
 
 	/** Treeview support functions */
-	virtual TSharedRef<ITableRow> GenerateTypeTreeRow(FPinTypeTreeItem InItem, const TSharedRef<STableViewBase>& OwnerTree);
-	void OnTypeSelectionChanged(FPinTypeTreeItem Selection, ESelectInfo::Type SelectInfo);
+	virtual TSharedRef<ITableRow> GenerateTypeTreeRow(FPinTypeTreeItem InItem, const TSharedRef<STableViewBase>& OwnerTree, bool bForSecondaryType);
+	void OnTypeSelectionChanged(FPinTypeTreeItem Selection, ESelectInfo::Type SelectInfo, bool bForSecondaryType);
 	void GetTypeChildren(FPinTypeTreeItem InItem, TArray<FPinTypeTreeItem>& OutChildren);
 
 	/** Listview support functions for sub-menu */
 	TSharedRef<ITableRow> GenerateObjectReferenceTreeRow(FObjectReferenceListItem InItem, const TSharedRef<STableViewBase>& OwnerTree);
-	void OnObjectReferenceSelectionChanged(FObjectReferenceListItem InItem, ESelectInfo::Type SelectInfo);
+	void OnObjectReferenceSelectionChanged(FObjectReferenceListItem InItem, ESelectInfo::Type SelectInfo, bool bForSecondaryType);
 
 	/** Reference to the menu content that's displayed when the type button is clicked on */
 	TSharedPtr<SMenuOwner> MenuContent;
-	virtual TSharedRef<SWidget>	GetMenuContent();
+	virtual TSharedRef<SWidget>	GetMenuContent(bool bForSecondaryType);
 
 	/** Type searching support */
 	FText SearchText;
@@ -135,9 +149,15 @@ protected:
 	/** Callback to get the tooltip text for the pin type combo box */
 	FText GetToolTipForComboBoxType() const;
 
+	/** Callback to get the tooltip text for the secondary pin type combo box */
+	FText GetToolTipForComboBoxSecondaryType() const;
+
 	/** Callback to get the tooltip for the array button widget */
 	FText GetToolTipForArrayWidget() const;
 
+	/** Callback to get the tooltip for the container type dropdown widget */
+	FText GetToolTipForContainerWidget() const;
+
 	/**
 	 * Helper function to create widget for the sub-menu
 	 *
@@ -149,7 +169,7 @@ protected:
 	TSharedRef<SWidget> CreateObjectReferenceWidget(FPinTypeTreeItem InItem, FEdGraphPinType& InPinType, const FSlateBrush* InIconBrush, FText InSimpleTooltip) const;
 
 	/** Gets the allowable object types for an tree item, used for building the sub-menu */
-	TSharedRef< SWidget > GetAllowedObjectTypes(FPinTypeTreeItem InItem);
+	TSharedRef< SWidget > GetAllowedObjectTypes(FPinTypeTreeItem InItem, bool bForSecondaryType);
 	
 	/**
 	 * When a pin type is selected, handle it
@@ -157,5 +177,5 @@ protected:
 	 * @param InItem				Item selected
 	 * @param InPinCategory			This is the PinType's category, must be provided separately as the PinType in the tree item is always Object Types for any object related type.
 	 */
-	void OnSelectPinType(FPinTypeTreeItem InItem, FString InPinCategory);
+	void OnSelectPinType(FPinTypeTreeItem InItem, FString InPinCategory, bool bForSecondaryType);
 };
diff --git a/Engine/Source/Editor/UnrealEd/Private/EdGraphUtilities.cpp b/Engine/Source/Editor/UnrealEd/Private/EdGraphUtilities.cpp
--- a/Engine/Source/Editor/UnrealEd/Private/EdGraphUtilities.cpp
+++ b/Engine/Source/Editor/UnrealEd/Private/EdGraphUtilities.cpp
@@ -461,6 +461,37 @@ void FEdGraphUtilities::CopyCommonState(UEdGraphNode* OldNode, UEdGraphNode* New
 	NewNode->NodeComment = OldNode->NodeComment;
 }
 
+bool FEdGraphUtilities::IsSetParam(const UFunction* Function, const FString& ParameterName)
+{
+	if (Function == nullptr)
+	{
+		return false;
+	}
+
+	const FString& RawMetaData = Function->GetMetaData(FBlueprintMetadata::MD_SetParam);
+	if (RawMetaData.IsEmpty())
+	{
+		return false;
+	}
+
+	TArray<FString> SetParamPinGroups;
+	{
+		RawMetaData.ParseIntoArray(SetParamPinGroups, TEXT(","), true);
+	}
+
+	for (FString& Entry : SetParamPinGroups)
+	{
+		TArray<FString> GroupEntries;
+		Entry.ParseIntoArray(GroupEntries, TEXT("|"), true);
+		if (GroupEntries.Contains(ParameterName))
+		{
+			return true;
+		}
+	}
+
+	return false;
+}
+
 void FEdGraphUtilities::RegisterVisualNodeFactory(TSharedPtr<FGraphPanelNodeFactory> NewFactory)
 {
 	VisualNodeFactories.Add(NewFactory);
diff --git a/Engine/Source/Editor/UnrealEd/Private/Kismet2/BlueprintEditorUtils.cpp b/Engine/Source/Editor/UnrealEd/Private/Kismet2/BlueprintEditorUtils.cpp
--- a/Engine/Source/Editor/UnrealEd/Private/Kismet2/BlueprintEditorUtils.cpp
+++ b/Engine/Source/Editor/UnrealEd/Private/Kismet2/BlueprintEditorUtils.cpp
@@ -8402,6 +8402,14 @@ const FSlateBrush* FBlueprintEditorUtils::GetIconFromPin( const FEdGraphPinType&
 	{
 		IconBrush = FEditorStyle::GetBrush(TEXT("Kismet.VariableList.ArrayTypeIcon"));
 	}
+	else if (PinType.bIsMap && PinType.PinCategory != K2Schema->PC_Exec)
+	{
+		IconBrush = FEditorStyle::GetBrush(TEXT("Kismet.VariableList.MapKeyTypeIcon"));
+	}
+	else if (PinType.bIsSet && PinType.PinCategory != K2Schema->PC_Exec)
+	{
+		IconBrush = FEditorStyle::GetBrush(TEXT("Kismet.VariableList.SetTypeIcon"));
+	}
 	else if( PinSubObject )
 	{
 		UClass* VarClass = FindObject<UClass>(ANY_PACKAGE, *PinSubObject->GetName());
@@ -8413,6 +8421,16 @@ const FSlateBrush* FBlueprintEditorUtils::GetIconFromPin( const FEdGraphPinType&
 	return IconBrush;
 }
 
+const FSlateBrush* FBlueprintEditorUtils::GetSecondaryIconFromPin(const FEdGraphPinType& PinType)
+{
+	const UEdGraphSchema_K2* K2Schema = GetDefault<UEdGraphSchema_K2>();
+	if (PinType.bIsMap && PinType.PinCategory != K2Schema->PC_Exec)
+	{
+		return FEditorStyle::GetBrush(TEXT("Kismet.VariableList.MapValueTypeIcon"));
+	}
+	return nullptr;
+}
+
 FText FBlueprintEditorUtils::GetFriendlyClassDisplayName(const UClass* Class)
 {
 	if (Class != nullptr)
@@ -8511,14 +8529,16 @@ void FBlueprintEditorUtils::HandleDisableEditableWhenInherited(UObject* Modified
 	}
 }
 
-void FBlueprintEditorUtils::BuildComponentInstancingData(UActorComponent* ComponentTemplate, FBlueprintCookedComponentInstancingData& OutData)
+struct FComponentInstancingDataUtils
 {
 	// Recursively gathers properties that differ from class/struct defaults, and fills out the cooked property list structure.
-	TFunction<void(UStruct*, const uint8*, const uint8*)> RecursivePropertyGatherLambda = [&OutData, &RecursivePropertyGatherLambda](UStruct* InStruct, const uint8* DataPtr, const uint8* DefaultDataPtr)
+	static void RecursivePropertyGather(UStruct* InStruct, const uint8* DataPtr, const uint8* DefaultDataPtr, FBlueprintCookedComponentInstancingData& OutData)
 	{
 		for (UProperty* Property = InStruct->PropertyLink; Property; Property = Property->PropertyLinkNext)
 		{
-			if (!Property->IsEditorOnlyProperty())
+			// Skip editor-only properties since they won't be compiled in a non-editor configuration. Also skip transient properties since they won't be serialized. 
+			if (!Property->IsEditorOnlyProperty()
+				&& !Property->HasAnyPropertyFlags(CPF_Transient | CPF_DuplicateTransient | CPF_NonPIEDuplicateTransient))
 			{
 				for (int32 Idx = 0; Idx < Property->ArrayDim; Idx++)
 				{
@@ -8534,7 +8554,7 @@ void FBlueprintEditorUtils::BuildComponentInstancingData(UActorComponent* Compon
 					{
 						int32 NumChangedProperties = OutData.ChangedPropertyList.Num();
 
-						RecursivePropertyGatherLambda(StructProperty->Struct, PropertyValue, DefaultPropertyValue);
+						RecursivePropertyGather(StructProperty->Struct, PropertyValue, DefaultPropertyValue, OutData);
 
 						// Prepend the struct property only if there is at least one changed sub-property.
 						if (NumChangedProperties < OutData.ChangedPropertyList.Num())
@@ -8563,7 +8583,7 @@ void FBlueprintEditorUtils::BuildComponentInstancingData(UActorComponent* Compon
 								{
 									int32 NumChangedArrayProperties = OutData.ChangedPropertyList.Num();
 
-									RecursivePropertyGatherLambda(InnerStructProperty->Struct, ArrayPropertyValue, DefaultArrayPropertyValue);
+									RecursivePropertyGather(InnerStructProperty->Struct, ArrayPropertyValue, DefaultArrayPropertyValue, OutData);
 
 									// Prepend the struct property only if there is at least one changed sub-property.
 									if (NumChangedArrayProperties < OutData.ChangedPropertyList.Num())
@@ -8571,7 +8591,7 @@ void FBlueprintEditorUtils::BuildComponentInstancingData(UActorComponent* Compon
 										OutData.ChangedPropertyList.Insert(ChangedArrayPropertyInfo, NumChangedArrayProperties);
 									}
 								}
-								else if(!ArrayProperty->Inner->Identical(ArrayPropertyValue, DefaultArrayPropertyValue, PPF_None))
+								else if (!ArrayProperty->Inner->Identical(ArrayPropertyValue, DefaultArrayPropertyValue, PPF_None))
 								{
 									// Emit the index of the individual array value that differs from the default value
 									OutData.ChangedPropertyList.Add(ChangedArrayPropertyInfo);
@@ -8601,20 +8621,29 @@ void FBlueprintEditorUtils::BuildComponentInstancingData(UActorComponent* Compon
 				}
 			}
 		}
-	};
+	}
+};
 
+void FBlueprintEditorUtils::BuildComponentInstancingData(UActorComponent* ComponentTemplate, FBlueprintCookedComponentInstancingData& OutData)
+{
 	if (ComponentTemplate)
 	{
 		UClass* ComponentTemplateClass = ComponentTemplate->GetClass();
 
 		// Gather the set of properties that differ from the template CDO.
 		OutData.ChangedPropertyList.Empty();
-		RecursivePropertyGatherLambda(ComponentTemplateClass, (uint8*)ComponentTemplate, (uint8*)ComponentTemplateClass->GetDefaultObject(false));
+		FComponentInstancingDataUtils::RecursivePropertyGather(ComponentTemplateClass, (uint8*)ComponentTemplate, (uint8*)ComponentTemplateClass->GetDefaultObject(false), OutData);
 
 		// Flag that cooked data has been built and is now considered to be valid.
 		OutData.bIsValid = true;
 	}
 }
 
+bool FBlueprintEditorUtils::ShouldEnableAdvancedContainers()
+{
+	UBlueprintEditorSettings const* Settings = GetDefault<UBlueprintEditorSettings>();
+	return Settings->bEnableAdvancedContainers;
+}
+
 #undef LOCTEXT_NAMESPACE
 
diff --git a/Engine/Source/Editor/UnrealEd/Private/Kismet2/ComponentEditorUtils.cpp b/Engine/Source/Editor/UnrealEd/Private/Kismet2/ComponentEditorUtils.cpp
--- a/Engine/Source/Editor/UnrealEd/Private/Kismet2/ComponentEditorUtils.cpp
+++ b/Engine/Source/Editor/UnrealEd/Private/Kismet2/ComponentEditorUtils.cpp
@@ -8,6 +8,7 @@
 #include "Engine/SimpleConstructionScript.h"
 #include "ScopedTransaction.h"
 #include "BlueprintEditorUtils.h"
+#include "KismetEditorUtilities.h"
 #include "Factories.h"
 #include "UnrealExporter.h"
 #include "GenericCommands.h"
@@ -86,13 +87,18 @@ protected:
 
 	virtual bool CanCreateClass(UClass* ObjectClass, bool& bOmitSubObjs) const override
 	{
-		// Only allow actor component types to be created
+		// Allow actor component types to be created
 		bool bCanCreate = ObjectClass->IsChildOf(UActorComponent::StaticClass());
 
-		// Also allow actor types to pass, in order to enable proper creation of actor component types as subobjects. The actor instance will be discarded after processing.
 		if (!bCanCreate)
 		{
-			bCanCreate = ObjectClass->IsChildOf(AActor::StaticClass());
+			// Also allow Blueprint-able actor types to pass, in order to enable proper creation of actor component types as subobjects. The actor instance will be discarded after processing.
+			bCanCreate = ObjectClass->IsChildOf(AActor::StaticClass()) && FKismetEditorUtilities::CanCreateBlueprintOfClass(ObjectClass);
+		}
+		else
+		{
+			// Actor component classes should not be abstract and must also be tagged as BlueprintSpawnable
+			bCanCreate = !ObjectClass->HasAnyClassFlags(CLASS_Abstract) && ObjectClass->HasMetaData(FBlueprintMetadata::MD_BlueprintSpawnableComponent);
 		}
 
 		return bCanCreate;
diff --git a/Engine/Source/Editor/UnrealEd/Private/ThumbnailRendering/BlueprintThumbnailRenderer.cpp b/Engine/Source/Editor/UnrealEd/Private/ThumbnailRendering/BlueprintThumbnailRenderer.cpp
--- a/Engine/Source/Editor/UnrealEd/Private/ThumbnailRendering/BlueprintThumbnailRenderer.cpp
+++ b/Engine/Source/Editor/UnrealEd/Private/ThumbnailRendering/BlueprintThumbnailRenderer.cpp
@@ -76,7 +76,15 @@ bool UBlueprintThumbnailRenderer::CanVisualizeAsset(UObject* Object)
 void UBlueprintThumbnailRenderer::Draw(UObject* Object, int32 X, int32 Y, uint32 Width, uint32 Height, FRenderTarget* RenderTarget, FCanvas* Canvas)
 {
 	UBlueprint* Blueprint = Cast<UBlueprint>(Object);
-	if (Blueprint && Blueprint->GeneratedClass)
+
+	// Strict validation - it may hopefully fix UE-35705.
+	const bool bIsBlueprintValid = IsValid(Blueprint)
+		&& IsValid(Blueprint->GeneratedClass)
+		&& Blueprint->bHasBeenRegenerated
+		//&& Blueprint->IsUpToDate() - This condition makes the thumbnail blank whenever the BP is dirty. It seems too strict.
+		&& !Blueprint->bBeingCompiled
+		&& !Blueprint->HasAnyFlags(RF_Transient);
+	if (bIsBlueprintValid)
 	{
 		TSharedRef<FBlueprintThumbnailScene> ThumbnailScene = ThumbnailScenes.EnsureThumbnailScene(Blueprint->GeneratedClass);
 
diff --git a/Engine/Source/Editor/UnrealEd/Public/EdGraphUtilities.h b/Engine/Source/Editor/UnrealEd/Public/EdGraphUtilities.h
--- a/Engine/Source/Editor/UnrealEd/Public/EdGraphUtilities.h
+++ b/Engine/Source/Editor/UnrealEd/Public/EdGraphUtilities.h
@@ -124,6 +124,9 @@ public:
 		void TraverseNodes(UEdGraphNode* Node);
 	};
 
+	/** returns true if the ParameterName is marked as a SetParam in the relevant Function */
+	static bool IsSetParam(const UFunction* Function, const FString& ParameterName);
+
 private:
 	static TArray< TSharedPtr<FGraphPanelNodeFactory> > VisualNodeFactories;
 	static TArray< TSharedPtr<FGraphPanelPinFactory> > VisualPinFactories;
diff --git a/Engine/Source/Editor/UnrealEd/Public/Kismet2/BlueprintEditorUtils.h b/Engine/Source/Editor/UnrealEd/Public/Kismet2/BlueprintEditorUtils.h
--- a/Engine/Source/Editor/UnrealEd/Public/Kismet2/BlueprintEditorUtils.h
+++ b/Engine/Source/Editor/UnrealEd/Public/Kismet2/BlueprintEditorUtils.h
@@ -1302,6 +1302,11 @@ public:
 	 */
 	static const struct FSlateBrush* GetIconFromPin(const FEdGraphPinType& PinType);
 
+	/**
+	 * Determine the best secondary icon icon to represent the given pin.
+	 */
+	static const struct FSlateBrush* GetSecondaryIconFromPin(const FEdGraphPinType& PinType);
+
 	/**
 	 * Generate component instancing data (for cooked builds).
 	 *
@@ -1311,6 +1316,9 @@ public:
 	 */
 	static void BuildComponentInstancingData(UActorComponent* ComponentTemplate, FBlueprintCookedComponentInstancingData& OutData);
 
+	/** Switch for hiding TMap/TSet support until stablized */
+	static bool ShouldEnableAdvancedContainers();
+
 protected:
 	// Removes all NULL graph references from the SubGraphs array and recurses thru the non-NULL ones
 	static void CleanNullGraphReferencesRecursive(UEdGraph* Graph);
diff --git a/Engine/Source/Programs/UnrealBuildTool/Configuration/UEBuildConfiguration.cs b/Engine/Source/Programs/UnrealBuildTool/Configuration/UEBuildConfiguration.cs
--- a/Engine/Source/Programs/UnrealBuildTool/Configuration/UEBuildConfiguration.cs
+++ b/Engine/Source/Programs/UnrealBuildTool/Configuration/UEBuildConfiguration.cs
@@ -242,10 +242,16 @@ namespace UnrealBuildTool
 		[XmlConfig]
 		public static bool bCompileWithPluginSupport;
 
-		/// <summary>
-		/// Whether to turn on logging for test/shipping builds
-		/// </summary>
-		[XmlConfig]
+        /// <summary>
+        /// Whether to include PerfCounters support
+        /// </summary>
+        [XmlConfig]
+        public static bool bWithPerfCounters;
+
+        /// <summary>
+        /// Whether to turn on logging for test/shipping builds
+        /// </summary>
+        [XmlConfig]
 		public static bool bUseLoggingInShipping;
 
 		/// <summary>
@@ -370,7 +376,8 @@ namespace UnrealBuildTool
 			bCompileSpeedTree = true;
 			bCompileWithStatsWithoutEngine = false;
 			bCompileWithPluginSupport = false;
-			bUseLoggingInShipping = false;
+            bWithPerfCounters = false;
+            bUseLoggingInShipping = false;
 			bUseChecksInShipping = false;
             bUseLauncherChecks = false;
 			bCompilePhysXVehicle = true;
diff --git a/Engine/Source/Programs/UnrealBuildTool/Configuration/UEBuildEditor.cs b/Engine/Source/Programs/UnrealBuildTool/Configuration/UEBuildEditor.cs
--- a/Engine/Source/Programs/UnrealBuildTool/Configuration/UEBuildEditor.cs
+++ b/Engine/Source/Programs/UnrealBuildTool/Configuration/UEBuildEditor.cs
@@ -71,8 +71,11 @@ namespace UnrealBuildTool
 			// Compile the engine
 			UEBuildConfiguration.bCompileAgainstEngine = true;
 
-			// Tag it as a 'Editor' build
-			OutCPPEnvironmentConfiguration.Definitions.Add("UE_EDITOR=1");
+            //enable PerfCounters
+            UEBuildConfiguration.bWithPerfCounters = true;
+
+            // Tag it as a 'Editor' build
+            OutCPPEnvironmentConfiguration.Definitions.Add("UE_EDITOR=1");
 		}
 	}
 }
diff --git a/Engine/Source/Programs/UnrealBuildTool/Configuration/UEBuildPlatform.cs b/Engine/Source/Programs/UnrealBuildTool/Configuration/UEBuildPlatform.cs
--- a/Engine/Source/Programs/UnrealBuildTool/Configuration/UEBuildPlatform.cs
+++ b/Engine/Source/Programs/UnrealBuildTool/Configuration/UEBuildPlatform.cs
@@ -244,7 +244,13 @@ namespace UnrealBuildTool
 					UEBuildConfiguration.bCompileWithPluginSupport = bValue;
 				}
 
-				bValue = UEBuildConfiguration.bCompilePhysXVehicle;
+                bValue = UEBuildConfiguration.bWithPerfCounters;
+                if (Ini.GetBool("/Script/BuildSettings.BuildSettings", "bWithPerfCounters", out bValue))
+                {
+                    UEBuildConfiguration.bWithPerfCounters = bValue;
+                }
+
+                bValue = UEBuildConfiguration.bCompilePhysXVehicle;
 				if (Ini.GetBool("/Script/BuildSettings.BuildSettings", "bCompilePhysXVehicle", out bValue))
 				{
 					UEBuildConfiguration.bCompilePhysXVehicle = bValue;
diff --git a/Engine/Source/Programs/UnrealBuildTool/Configuration/UEBuildServer.cs b/Engine/Source/Programs/UnrealBuildTool/Configuration/UEBuildServer.cs
--- a/Engine/Source/Programs/UnrealBuildTool/Configuration/UEBuildServer.cs
+++ b/Engine/Source/Programs/UnrealBuildTool/Configuration/UEBuildServer.cs
@@ -69,6 +69,9 @@ namespace UnrealBuildTool
 
 			// Compile the engine
 			UEBuildConfiguration.bCompileAgainstEngine = true;
+			
+			//enable PerfCounters
+            UEBuildConfiguration.bWithPerfCounters = true;
 
 			// Tag it as a 'Server' build
 			OutCPPEnvironmentConfiguration.Definitions.Add("UE_SERVER=1");
diff --git a/Engine/Source/Programs/UnrealBuildTool/Configuration/UEBuildTarget.cs b/Engine/Source/Programs/UnrealBuildTool/Configuration/UEBuildTarget.cs
--- a/Engine/Source/Programs/UnrealBuildTool/Configuration/UEBuildTarget.cs
+++ b/Engine/Source/Programs/UnrealBuildTool/Configuration/UEBuildTarget.cs
@@ -4021,6 +4021,15 @@ namespace UnrealBuildTool
 				GlobalCompileEnvironment.Config.Definitions.Add("WITH_PLUGIN_SUPPORT=0");
 			}
 
+            if (UEBuildConfiguration.bWithPerfCounters)
+            {
+                GlobalCompileEnvironment.Config.Definitions.Add("WITH_PERFCOUNTERS=1");
+            }
+            else
+            {
+                GlobalCompileEnvironment.Config.Definitions.Add("WITH_PERFCOUNTERS=0");
+            }
+
 			if (UEBuildConfiguration.bUseLoggingInShipping)
 			{
 				GlobalCompileEnvironment.Config.Definitions.Add("USE_LOGGING_IN_SHIPPING=1");
diff --git a/Engine/Source/Runtime/Core/Public/Containers/Set.h b/Engine/Source/Runtime/Core/Public/Containers/Set.h
--- a/Engine/Source/Runtime/Core/Public/Containers/Set.h
+++ b/Engine/Source/Runtime/Core/Public/Containers/Set.h
@@ -1389,7 +1389,6 @@ public:
 	int32 AddUninitialized(const FScriptSetLayout& Layout)
 	{
 		int32 Result = Elements.AddUninitialized(Layout.SparseArrayLayout);
-		++HashSize;
 		return Result;
 	}
 
@@ -1409,8 +1408,6 @@ public:
 				GetTypedHash(HashIndex) = FSetElementId();
 			}
 
-			int32 NumBytesPerSetElement = Layout.Size;
-
 			// Add the existing elements to the new hash.
 			int32 Index = 0;
 			int32 Count = Elements.Num();
@@ -1439,6 +1436,84 @@ public:
 		}
 	}
 
+	uint8* Find(const void* Element, const FScriptSetLayout& Layout, TFunctionRef<uint32 (const void*)> GetKeyHash, TFunctionRef<bool (const void*, const void*)> EqualityFn)
+	{
+		if (Elements.Num())
+		{
+			const uint32 ElementHash = GetKeyHash(Element);
+			const int32  HashIndex = ElementHash & (HashSize - 1);
+
+			uint8* CurrentElement = nullptr;
+			for (FSetElementId ElementId = GetTypedHash(HashIndex);
+				ElementId.IsValidId();
+				ElementId = GetHashNextIdRef(CurrentElement, Layout))
+			{
+				CurrentElement = (uint8*)Elements.GetData(ElementId, Layout.SparseArrayLayout);
+				if (EqualityFn(Element, CurrentElement))
+				{
+					return CurrentElement;
+				}
+			}
+		}
+
+		return nullptr;
+	}
+
+	bool Add(const void* Element, const FScriptSetLayout& Layout, TFunctionRef<uint32(const void*)> GetKeyHash, TFunctionRef<bool(const void*, const void*)> EqualityFn, TFunctionRef<void(void*)> ConstructFn)
+	{
+		// Minor efficiency concern: we hash the element both here in the Find() call and below
+		// when we link the new element into the set
+		bool bIsAlreadyInSet = Find(Element, Layout, GetKeyHash, EqualityFn) != nullptr;
+		if (!bIsAlreadyInSet)
+		{
+			// add the set element
+			FSetElementId	ElementId(AddUninitialized(Layout));
+			void* CurrentElement = Elements.GetData(ElementId, Layout.SparseArrayLayout);
+
+			ConstructFn(CurrentElement);
+
+			const int32 DesiredHashSize = FDefaultSetAllocator::GetNumberOfHashBuckets(Num());
+			if (!HashSize ||
+				HashSize < DesiredHashSize)
+			{
+				// rehash, this will link in our new element if needed:
+				Rehash(Layout, GetKeyHash);
+			}
+			else
+			{
+				// link the new element into the set:
+				const uint32	ElementHash = GetKeyHash(Element);
+				const int32		HashIndex = ElementHash & (HashSize - 1);
+				FSetElementId& TypedHash = GetTypedHash(HashIndex);
+				GetHashNextIdRef(CurrentElement, Layout) = TypedHash;
+				TypedHash = ElementId;
+			}
+		}
+		return !bIsAlreadyInSet;
+	}
+
+	void Remove(const void* Element, const FScriptSetLayout& Layout, TFunctionRef<uint32(const void*)> GetKeyHash, TFunctionRef<bool(const void*, const void*)> EqualityFn)
+	{
+		if (Elements.Num())
+		{
+			const uint32	ElementHash = GetKeyHash(Element);
+			const int32		HashIndex = ElementHash & (HashSize - 1);
+
+			void* CurrentElement = nullptr;
+			for(FSetElementId ElementId = GetTypedHash(HashIndex);
+				ElementId.IsValidId();
+				ElementId = GetHashNextIdRef(CurrentElement, Layout))
+			{
+				CurrentElement = (uint8*)Elements.GetData(ElementId, Layout.SparseArrayLayout);
+				if(EqualityFn(Element, CurrentElement))
+				{
+					RemoveAt(ElementId, Layout);
+					return;
+				}
+			}
+		}
+	}
+
 private:
 	typedef FDefaultSetAllocator Allocator;
 	typedef Allocator::HashAllocator::ForElementType<FSetElementId> HashType;
diff --git a/Engine/Source/Runtime/Core/Public/Misc/Build.h b/Engine/Source/Runtime/Core/Public/Misc/Build.h
--- a/Engine/Source/Runtime/Core/Public/Misc/Build.h
+++ b/Engine/Source/Runtime/Core/Public/Misc/Build.h
@@ -88,6 +88,11 @@
 	#error UBT should always define WITH_PLUGIN_SUPPORT to be 0 or 1
 #endif
 
+ /** Enable perf counters */
+#ifndef WITH_PERFCOUNTERS
+	#define WITH_PERFCOUNTERS		0
+#endif
+
 /**
  * Unreal Header Tool requires extra data stored in the structure of a few core files. This enables some ifdef hacks to make this work. 
  * Set via UBT, do not modify directly
diff --git a/Engine/Source/Runtime/Core/Public/UObject/BlueprintsObjectVersion.h b/Engine/Source/Runtime/Core/Public/UObject/BlueprintsObjectVersion.h
--- a/Engine/Source/Runtime/Core/Public/UObject/BlueprintsObjectVersion.h
+++ b/Engine/Source/Runtime/Core/Public/UObject/BlueprintsObjectVersion.h
@@ -13,6 +13,7 @@ struct CORE_API FBlueprintsObjectVersion
 		ArrayGetByRefUpgrade,
 		EdGraphPinOptimized,
 		AllowDeletionConformed,
+		AdvancedContainerSupport,
 
 		// -----<new versions can be added above this line>-------------------------------------------------
 		VersionPlusOne,
diff --git a/Engine/Source/Runtime/CoreUObject/Private/Blueprint/BlueprintSupport.cpp b/Engine/Source/Runtime/CoreUObject/Private/Blueprint/BlueprintSupport.cpp
--- a/Engine/Source/Runtime/CoreUObject/Private/Blueprint/BlueprintSupport.cpp
+++ b/Engine/Source/Runtime/CoreUObject/Private/Blueprint/BlueprintSupport.cpp
@@ -475,7 +475,7 @@ bool FLinkerLoad::DeferPotentialCircularImport(const int32 Index)
 				{
 					Import.XObject = MakeImportPlaceholder<ULinkerPlaceholderClass>(LinkerRoot, *Import.ObjectName.ToString(), Index);
 				}
-				else if (ImportClass->IsChildOf<UFunction>())
+				else if (ImportClass->IsChildOf<UFunction>() && Import.OuterIndex.IsImport())
 				{
 					const int32 OuterImportIndex = Import.OuterIndex.ToImport();
 					// @TODO: if the sole reason why we have ULinkerPlaceholderFunction 
diff --git a/Engine/Source/Runtime/CoreUObject/Private/UObject/Class.cpp b/Engine/Source/Runtime/CoreUObject/Private/UObject/Class.cpp
--- a/Engine/Source/Runtime/CoreUObject/Private/UObject/Class.cpp
+++ b/Engine/Source/Runtime/CoreUObject/Private/UObject/Class.cpp
@@ -2479,8 +2479,8 @@ UObject* UClass::CreateDefaultObject()
 				ClassDefaultObject = StaticAllocateObject(this, GetOuter(), NAME_None, EObjectFlags(RF_Public|RF_ClassDefaultObject|RF_ArchetypeObject));
 				check(ClassDefaultObject);
 				// Blueprint CDOs have their properties always initialized.
-				const bool bShouldInitilizeProperties = !HasAnyClassFlags(CLASS_Native | CLASS_Intrinsic);
-				(*ClassConstructor)(FObjectInitializer(ClassDefaultObject, ParentDefaultObject, false, bShouldInitilizeProperties));
+				const bool bShouldInitializeProperties = !HasAnyClassFlags(CLASS_Native | CLASS_Intrinsic);
+				(*ClassConstructor)(FObjectInitializer(ClassDefaultObject, ParentDefaultObject, false, bShouldInitializeProperties));
 			}
 		}
 	}
diff --git a/Engine/Source/Runtime/CoreUObject/Private/UObject/PropertyArray.cpp b/Engine/Source/Runtime/CoreUObject/Private/UObject/PropertyArray.cpp
--- a/Engine/Source/Runtime/CoreUObject/Private/UObject/PropertyArray.cpp
+++ b/Engine/Source/Runtime/CoreUObject/Private/UObject/PropertyArray.cpp
@@ -94,7 +94,10 @@ void UArrayProperty::SerializeItem( FArchive& Ar, void* Value, void const* Defau
 
 			auto CanSerializeFromStructWithDifferentName = [](const FArchive& InAr, const FPropertyTag& PropertyTag, const UStructProperty* StructProperty)
 			{
-				return PropertyTag.StructGuid.IsValid() && StructProperty && StructProperty->Struct && (PropertyTag.StructGuid == StructProperty->Struct->GetCustomGuid() && (StructProperty->Struct->StructFlags & STRUCT_SerializeFromMismatchedTag));
+				return PropertyTag.StructGuid.IsValid()
+					&& StructProperty 
+					&& StructProperty->Struct 
+					&& (PropertyTag.StructGuid == StructProperty->Struct->GetCustomGuid());
 			};
 
 			// Check if the Inner property can successfully serialize, the type may have changed
diff --git a/Engine/Source/Runtime/CoreUObject/Private/UObject/PropertySet.cpp b/Engine/Source/Runtime/CoreUObject/Private/UObject/PropertySet.cpp
--- a/Engine/Source/Runtime/CoreUObject/Private/UObject/PropertySet.cpp
+++ b/Engine/Source/Runtime/CoreUObject/Private/UObject/PropertySet.cpp
@@ -387,6 +387,13 @@ void USetProperty::AddReferencedObjects(UObject* InThis, FReferenceCollector& Co
 	Super::AddReferencedObjects(This, Collector);
 }
 
+FString USetProperty::GetCPPMacroType(FString& ExtendedTypeText) const
+{
+	checkSlow(ElementProp);
+	ExtendedTypeText = FString::Printf(TEXT("%s"), *ElementProp->GetCPPType());
+	return TEXT("TSET");
+}
+
 FString USetProperty::GetCPPType(FString* ExtendedTypeText, uint32 CPPExportFlags) const
 {
 	checkSlow(ElementProp);
@@ -402,11 +409,10 @@ FString USetProperty::GetCPPType(FString* ExtendedTypeText, uint32 CPPExportFlag
 	return TEXT("TSet");
 }
 
-FString USetProperty::GetCPPMacroType( FString& ExtendedTypeText ) const
+FString USetProperty::GetCPPTypeForwardDeclaration() const
 {
 	checkSlow(ElementProp);
-	ExtendedTypeText = FString::Printf(TEXT("%s"), *ElementProp->GetCPPType());
-	return TEXT("TSET");
+	return ElementProp->GetCPPTypeForwardDeclaration();
 }
 
 void USetProperty::ExportTextItem(FString& ValueStr, const void* PropertyValue, const void* DefaultValue, UObject* Parent, int32 PortFlags, UObject* ExportRootScope) const
diff --git a/Engine/Source/Runtime/CoreUObject/Private/UObject/ScriptCore.cpp b/Engine/Source/Runtime/CoreUObject/Private/UObject/ScriptCore.cpp
--- a/Engine/Source/Runtime/CoreUObject/Private/UObject/ScriptCore.cpp
+++ b/Engine/Source/Runtime/CoreUObject/Private/UObject/ScriptCore.cpp
@@ -3,7 +3,6 @@
 /*=============================================================================
 	ScriptCore.cpp: Kismet VM execution and support code.
 =============================================================================*/
-
 #include "CoreUObjectPrivate.h"
 #include "MallocProfiler.h"
 #include "HotReloadInterface.h"
@@ -12,14 +11,11 @@
 DEFINE_LOG_CATEGORY(LogScriptFrame);
 DEFINE_LOG_CATEGORY_STATIC(LogScriptCore, Log, All);
 
-DECLARE_CYCLE_STAT(TEXT("Blueprint Time"),STAT_BlueprintTime,STATGROUP_Game);
+DECLARE_CYCLE_STAT(TEXT("Blueprint Time"), STAT_BlueprintTime, STATGROUP_Game);
 
 #define LOCTEXT_NAMESPACE "ScriptCore"
 
 #if TOTAL_OVERHEAD_SCRIPT_STATS
-COREUOBJECT_API FBlueprintEventTimer::FScopedVMTimer* FBlueprintEventTimer::ActiveVMTimer = nullptr;
-COREUOBJECT_API FBlueprintEventTimer::FPausableScopeTimer* FBlueprintEventTimer::ActiveTimer = nullptr;
-
 DEFINE_STAT(STAT_ScriptVmTime_Total);
 DEFINE_STAT(STAT_ScriptNativeTime_Total);
 #endif //TOTAL_OVERHEAD_SCRIPT_STATS
diff --git a/Engine/Source/Runtime/CoreUObject/Public/UObject/ObjectMacros.h b/Engine/Source/Runtime/CoreUObject/Public/UObject/ObjectMacros.h
--- a/Engine/Source/Runtime/CoreUObject/Public/UObject/ObjectMacros.h
+++ b/Engine/Source/Runtime/CoreUObject/Public/UObject/ObjectMacros.h
@@ -1541,6 +1541,13 @@ public: \
 		UPackage* PrivateStaticClassOuter = FindOrConstructDynamicTypePackage(Package); \
 		UClass* PrivateStaticClass = Cast<UClass>(StaticFindObjectFast(UClass::StaticClass(), PrivateStaticClassOuter, (TCHAR*)ClassName)); \
 		if (!PrivateStaticClass) \
+		{ \
+			/* the class could be created while its parent creation, so make sure, the parent is already created.*/ \
+			TClass::Super::StaticClass(); \
+			TClass::WithinClass::StaticClass(); \
+			PrivateStaticClass = Cast<UClass>(StaticFindObjectFast(UClass::StaticClass(), PrivateStaticClassOuter, (TCHAR*)ClassName)); \
+		} \
+		if (!PrivateStaticClass) \
 		{ \
 			/* this could be handled with templates, but we want it external to avoid code bloat */ \
 			GetPrivateStaticClassBody( \
diff --git a/Engine/Source/Runtime/CoreUObject/Public/UObject/Script.h b/Engine/Source/Runtime/CoreUObject/Public/UObject/Script.h
--- a/Engine/Source/Runtime/CoreUObject/Public/UObject/Script.h
+++ b/Engine/Source/Runtime/CoreUObject/Public/UObject/Script.h
@@ -460,12 +460,30 @@ private:
 	struct FBlueprintEventTimer
 	{
 		struct FPausableScopeTimer;
-		COREUOBJECT_API static FPausableScopeTimer* ActiveTimer;
+		struct FScopedVMTimer;
+
+		class FThreadedTimerManager : public TThreadSingleton<FThreadedTimerManager>
+		{
+		public:
+			FThreadedTimerManager()
+				: ActiveTimer(nullptr)
+				, ActiveVMScope(nullptr)
+			{}
+
+			FPausableScopeTimer* ActiveTimer;
+
+			// We need to keep track of the current VM timer because we only want to
+			// track time while 'in' the VM. We use this to detect whether we're running
+			// script or just doing RPC:
+			FScopedVMTimer* ActiveVMScope;
+		};
 
 		struct FPausableScopeTimer
 		{
 			FPausableScopeTimer() 
 			{ 
+				FPausableScopeTimer*& ActiveTimer = FThreadedTimerManager::Get().ActiveTimer;
+
 				double CurrentTime = FPlatformTime::Seconds();
 				if (ActiveTimer)
 				{
@@ -485,7 +503,7 @@ private:
 				{
 					PreviouslyActiveTimer->Resume();
 				}
-				ActiveTimer = PreviouslyActiveTimer;
+				FThreadedTimerManager::Get().ActiveTimer = PreviouslyActiveTimer;
 			}
 
 			void Pause(double CurrentTime) { TotalTime += CurrentTime - StartTime; }
@@ -498,27 +516,23 @@ private:
 			double StartTime;
 		};
 
-		// We need to keep track of the current VM timer because we only want to
-		// track time while 'in' the VM. We use this to detect whether we're running
-		// script or just doing RPC:
-		struct FScopedVMTimer;
-		COREUOBJECT_API static FScopedVMTimer* ActiveVMTimer;
-
 		struct FScopedVMTimer
 		{
 			FScopedVMTimer()
 				: Timer()
-				, VMParent(ActiveVMTimer)
+				, VMParent(nullptr)
 			{
+				FScopedVMTimer*& ActiveVMTimer = FThreadedTimerManager::Get().ActiveVMScope;
+				VMParent = ActiveVMTimer;
+
 				ActiveVMTimer = this;
 			}
 
 			~FScopedVMTimer()
 			{
 				INC_FLOAT_STAT_BY(STAT_ScriptVmTime_Total, Timer.Stop() * 1000.0);
-				ActiveVMTimer = VMParent;
+				FThreadedTimerManager::Get().ActiveVMScope = VMParent;
 			}
-
 			FPausableScopeTimer Timer;
 			FScopedVMTimer* VMParent;
 		};
@@ -527,14 +541,13 @@ private:
 		{
 			FScopedNativeTimer()
 				: Timer()
-			{
-			}
+			{}
 
 			~FScopedNativeTimer()
 			{
 				// only track native time when in a VM scope, RPC time
 				// can be tracked by the online system or whatever is making RPCs:
-				if (ActiveVMTimer)
+				if (FThreadedTimerManager::Get().ActiveVMScope)
 				{
 					INC_FLOAT_STAT_BY(STAT_ScriptNativeTime_Total, Timer.Stop()* 1000.0);
 				}
diff --git a/Engine/Source/Runtime/CoreUObject/Public/UObject/UnrealType.h b/Engine/Source/Runtime/CoreUObject/Public/UObject/UnrealType.h
--- a/Engine/Source/Runtime/CoreUObject/Public/UObject/UnrealType.h
+++ b/Engine/Source/Runtime/CoreUObject/Public/UObject/UnrealType.h
@@ -2642,6 +2642,7 @@ public:
 	// UProperty interface
 	virtual FString GetCPPMacroType(FString& ExtendedTypeText) const  override;
 	virtual FString GetCPPType(FString* ExtendedTypeText, uint32 CPPExportFlags) const override;
+	virtual FString GetCPPTypeForwardDeclaration() const override;
 	virtual void LinkInternal(FArchive& Ar) override;
 	virtual bool Identical(const void* A, const void* B, uint32 PortFlags) const override;
 	virtual void SerializeItem(FArchive& Ar, void* Value, void const* Defaults) const override;
@@ -3637,6 +3638,55 @@ public:
 		return Result;
 	}
 
+	/** Finds element from hash, rather than linearly searching */
+	FORCEINLINE uint8* FindElementFromHash(const void* ElementToFind) const
+	{
+		UProperty* LocalElementPropForCapture = ElementProp;
+		return Set->Find(
+				ElementToFind, 
+				SetLayout, 
+				[LocalElementPropForCapture](const void* Element) { return LocalElementPropForCapture->GetValueTypeHash(Element); },
+				[LocalElementPropForCapture](const void* A, const void* B) { return LocalElementPropForCapture->Identical(A, B); }
+			);
+	}
+
+	/** Adds the element to the set, returning true if the element was added, or false if the element was already present */
+	bool AddElement(const void* ElementToAdd)
+	{
+		UProperty* LocalElementPropForCapture = ElementProp;
+		FScriptSetLayout& LocalSetLayoutForCapture = SetLayout;
+		return Set->Add(
+			ElementToAdd,
+			SetLayout,
+			[LocalElementPropForCapture](const void* Element) { return LocalElementPropForCapture->GetValueTypeHash(Element); },
+			[LocalElementPropForCapture](const void* A, const void* B) { return LocalElementPropForCapture->Identical(A, B); },
+			[LocalElementPropForCapture, ElementToAdd, LocalSetLayoutForCapture](void* NewElement)
+			{
+				if (LocalElementPropForCapture->PropertyFlags & CPF_ZeroConstructor)
+				{
+					FMemory::Memzero(NewElement, LocalSetLayoutForCapture.Size);
+				}
+				else
+				{
+					LocalElementPropForCapture->InitializeValue(NewElement);
+				}
+
+				LocalElementPropForCapture->CopySingleValueToScriptVM(NewElement, ElementToAdd);
+			}
+		);
+	}
+
+	/** Remoes the elemnt from the set */
+	void RemoveElement(const void* ElementToRemove)
+	{
+		UProperty* LocalElementPropForCapture = ElementProp;
+		Set->Remove(
+			ElementToRemove, 
+			SetLayout, 
+			[LocalElementPropForCapture](const void* Element) { return LocalElementPropForCapture->GetValueTypeHash(Element); },
+			[LocalElementPropForCapture](const void* A, const void* B) { return LocalElementPropForCapture->Identical(A, B); }
+		);
+	}
 
 private:
 	/**
diff --git a/Engine/Source/Runtime/Engine/Classes/EdGraph/EdGraphPin.h b/Engine/Source/Runtime/Engine/Classes/EdGraph/EdGraphPin.h
--- a/Engine/Source/Runtime/Engine/Classes/EdGraph/EdGraphPin.h
+++ b/Engine/Source/Runtime/Engine/Classes/EdGraph/EdGraphPin.h
@@ -54,6 +54,62 @@ FORCEINLINE FArchive& operator<<(FArchive& Ar, FSimpleMemberReference& Data)
 	return Ar;
 }
 
+/**
+  * Struct used to define information for terminal types, e.g. types that can be contained
+  * by a container. Currently can represent strong/weak references to a type (only UObjects), 
+  * a structure, or a primitive. Support for "Container of Containers" is done by wrapping 
+  * a structure, rather than implicitly defining names for containers.
+  */
+USTRUCT()
+struct FEdGraphTerminalType
+{
+	GENERATED_USTRUCT_BODY()
+
+	FEdGraphTerminalType()
+		: TerminalCategory()
+		, TerminalSubCategory()
+		, TerminalSubCategoryObject(nullptr)
+		, bTerminalIsConst(false)
+		, bTerminalIsWeakPointer(false)
+	{
+	}
+
+	/** Category */
+	UPROPERTY()
+	FString TerminalCategory;
+
+	/** Sub-category */
+	UPROPERTY()
+	FString TerminalSubCategory;
+
+	/** Sub-category object */
+	UPROPERTY()
+	TWeakObjectPtr<UObject> TerminalSubCategoryObject;
+
+	/** Whether or not this pin is a immutable const value */
+	UPROPERTY()
+	bool bTerminalIsConst;
+
+	/** Whether or not this is a weak reference */
+	UPROPERTY()
+	bool bTerminalIsWeakPointer;
+
+	ENGINE_API friend FArchive& operator<<(FArchive& Ar, FEdGraphTerminalType& P);
+};
+
+inline bool operator!= (const FEdGraphTerminalType& A, const FEdGraphTerminalType& B)
+{
+	return A.TerminalCategory != B.TerminalCategory
+		|| A.TerminalSubCategory != B.TerminalSubCategory
+		|| A.TerminalSubCategoryObject != B.TerminalSubCategoryObject
+		|| A.bTerminalIsConst != B.bTerminalIsConst
+		|| A.bTerminalIsWeakPointer != B.bTerminalIsWeakPointer;
+}
+
+inline bool operator==(const FEdGraphTerminalType& A, const FEdGraphTerminalType& B)
+{
+	return !(A != B);
+}
 
 /** Struct used to define the type of information carried on this pin */
 USTRUCT()
@@ -77,6 +133,18 @@ struct FEdGraphPinType
 	UPROPERTY()
 	FSimpleMemberReference PinSubCategoryMemberReference;
 
+	/** Data used to determine value types when bIsMap is true */
+	UPROPERTY()
+	FEdGraphTerminalType PinValueType;
+
+	/** Whether or not this pin represents a map of keys to values */
+	UPROPERTY()
+	bool bIsMap;
+
+	/** Whether or not this pin represents a set of (unique) values */
+	UPROPERTY()
+	bool bIsSet;
+
 	/** Whether or not this pin represents an array of values */
 	UPROPERTY()
 	bool bIsArray;
@@ -93,10 +161,14 @@ struct FEdGraphPinType
 	UPROPERTY()
 	bool bIsWeakPointer;
 
+	FORCEINLINE bool IsContainer() const { return bIsMap || bIsSet || bIsArray; }
+
 public:
 	FEdGraphPinType() 
 	{
-		PinSubCategoryObject = NULL;
+		PinSubCategoryObject = nullptr;
+		bIsMap = false;
+		bIsSet = false;
 		bIsArray = false;
 		bIsReference = false;
 		bIsConst = false;
@@ -107,6 +179,8 @@ public:
 		PinCategory = InPinCategory;
 		PinSubCategory = InPinSubCategory;
 		PinSubCategoryObject = InPinSubCategoryObject;
+		bIsMap = false;
+		bIsSet = false;
 		bIsArray = bInIsArray;
 		bIsReference = bInIsReference;
 		bIsConst = false;
@@ -117,6 +191,9 @@ public:
 		return (PinCategory == Other.PinCategory)
 			&& (PinSubCategory == Other.PinSubCategory)
 			&& (PinSubCategoryObject == Other.PinSubCategoryObject)
+			&& (PinValueType == Other.PinValueType)
+			&& (bIsMap == Other.bIsMap)
+			&& (bIsSet == Other.bIsSet)
 			&& (bIsArray == Other.bIsArray)
 			&& (bIsReference == Other.bIsReference)
 			&& (bIsWeakPointer == Other.bIsWeakPointer)
@@ -128,6 +205,9 @@ public:
 		return (PinCategory != Other.PinCategory) 
 			|| (PinSubCategory != Other.PinSubCategory) 
 			|| (PinSubCategoryObject != Other.PinSubCategoryObject) 
+			|| (PinValueType != Other.PinValueType)
+			|| (bIsMap != Other.bIsMap)
+			|| (bIsSet != Other.bIsSet)
 			|| (bIsArray != Other.bIsArray) 
 			|| (bIsReference != Other.bIsReference)
 			|| (bIsWeakPointer != Other.bIsWeakPointer)
@@ -138,8 +218,11 @@ public:
 	{
 		PinCategory.Empty();
 		PinSubCategory.Empty();
-		PinSubCategoryObject = NULL;
+		PinSubCategoryObject = nullptr;
+		PinValueType = FEdGraphTerminalType();
 		PinSubCategoryMemberReference.Reset();
+		bIsMap = false;
+		bIsSet = false;
 		bIsArray = false;
 		bIsReference = false;
 		bIsWeakPointer = false;
@@ -414,6 +497,7 @@ public:
 	ENGINE_API const FString& GetName() const { return PinName; }
 	ENGINE_API UEdGraphNode* GetOuter() const { return GetOwningNodeUnchecked(); }
 	ENGINE_API bool IsPendingKill() const {	return bWasTrashed; }
+	ENGINE_API FEdGraphTerminalType GetPrimaryTerminalType() const;
 
 	/** Verification that all pins have been destroyed after shutting down */
 	ENGINE_API static void ShutdownVerification();
diff --git a/Engine/Source/Runtime/Engine/Classes/Engine/Blueprint.h b/Engine/Source/Runtime/Engine/Classes/Engine/Blueprint.h
--- a/Engine/Source/Runtime/Engine/Classes/Engine/Blueprint.h
+++ b/Engine/Source/Runtime/Engine/Classes/Engine/Blueprint.h
@@ -509,6 +509,16 @@ public:
 	// User Defined Structures, the blueprint depends on
 	TSet<TWeakObjectPtr<UStruct>> CachedUDSDependencies;
 
+	enum class EIsBPNonReducible : uint8
+	{
+		Unkown,
+		Yes,
+		No,
+	};
+
+	// Cached information if the BP contains any non-reducible functions (that can benefit from nativization).
+	EIsBPNonReducible bHasAnyNonReducibleFunction;
+
 	// If this BP is just a duplicate created for a specific compilation, the reference to original GeneratedClass is needed
 	UPROPERTY(transient, duplicatetransient)
 	UClass* OriginalClass;
diff --git a/Engine/Source/Runtime/Engine/Private/ActorConstruction.cpp b/Engine/Source/Runtime/Engine/Private/ActorConstruction.cpp
--- a/Engine/Source/Runtime/Engine/Private/ActorConstruction.cpp
+++ b/Engine/Source/Runtime/Engine/Private/ActorConstruction.cpp
@@ -857,6 +857,9 @@ UActorComponent* AActor::CreateComponentFromTemplateData(const FBlueprintCookedC
 
 			// Set this flag to emulate things that would happen in the SDO case when this flag is set (e.g. - not setting 'bHasBeenCreated').
 			ArPortFlags |= PPF_Duplicate;
+
+			// Set this flag to ensure that we also serialize any deprecated properties.
+			ArPortFlags |= PPF_UseDeprecatedProperties;
 		}
 	};
 
diff --git a/Engine/Source/Runtime/Engine/Private/Animation/AnimInstanceProxy.cpp b/Engine/Source/Runtime/Engine/Private/Animation/AnimInstanceProxy.cpp
--- a/Engine/Source/Runtime/Engine/Private/Animation/AnimInstanceProxy.cpp
+++ b/Engine/Source/Runtime/Engine/Private/Animation/AnimInstanceProxy.cpp
@@ -1407,7 +1407,7 @@ void FAnimInstanceProxy::BindNativeDelegates()
 					{
 						if(TransitionExit.CanTakeDelegateIndex != INDEX_NONE)
 						{
-							// In case the state machine hasn't been initilized, we need to re-get the desc
+							// In case the state machine hasn't been initialized, we need to re-get the desc
 							const FBakedAnimationStateMachine* MachineDesc = GetMachineDescription(AnimClassInterface, StateMachine);
 							const FAnimationTransitionBetweenStates& Transition = MachineDesc->Transitions[TransitionExit.TransitionIndex];
 							const FBakedAnimationState& BakedState = MachineDesc->States[Transition.NextState];
diff --git a/Engine/Source/Runtime/Engine/Private/Blueprint.cpp b/Engine/Source/Runtime/Engine/Private/Blueprint.cpp
--- a/Engine/Source/Runtime/Engine/Private/Blueprint.cpp
+++ b/Engine/Source/Runtime/Engine/Private/Blueprint.cpp
@@ -312,6 +312,7 @@ UBlueprint::UBlueprint(const FObjectInitializer& ObjectInitializer)
 #if WITH_EDITORONLY_DATA
 	, bDuplicatingReadOnly(false)
 	, bCachedDependenciesUpToDate(false)
+	, bHasAnyNonReducibleFunction(EIsBPNonReducible::Unkown)
 #endif
 {
 }
diff --git a/Engine/Source/Runtime/Engine/Private/BlueprintGeneratedClass.cpp b/Engine/Source/Runtime/Engine/Private/BlueprintGeneratedClass.cpp
--- a/Engine/Source/Runtime/Engine/Private/BlueprintGeneratedClass.cpp
+++ b/Engine/Source/Runtime/Engine/Private/BlueprintGeneratedClass.cpp
@@ -1250,6 +1250,12 @@ void FBlueprintCookedComponentInstancingData::LoadCachedPropertyDataForSerializa
 			ArCustomPropertyList = InPropertyList;
 			ArUseCustomPropertyList = true;
 			ArWantBinaryPropertySerialization = true;
+
+			// Set this flag to emulate things that would normally happen in the SDO case when this flag is set. This is needed to ensure consistency with serialization during instancing.
+			ArPortFlags |= PPF_Duplicate;
+
+			// Set this flag to ensure that we also serialize any deprecated properties.
+			ArPortFlags |= PPF_UseDeprecatedProperties;
 		}
 	};
 
diff --git a/Engine/Source/Runtime/Engine/Private/EdGraph/EdGraphPin.cpp b/Engine/Source/Runtime/Engine/Private/EdGraph/EdGraphPin.cpp
--- a/Engine/Source/Runtime/Engine/Private/EdGraph/EdGraphPin.cpp
+++ b/Engine/Source/Runtime/Engine/Private/EdGraph/EdGraphPin.cpp
@@ -1,6 +1,8 @@
 // Copyright 1998-2016 Epic Games, Inc. All Rights Reserved.
 
 #include "EnginePrivate.h"
+
+#include "BlueprintsObjectVersion.h"
 #include "EdGraph/EdGraph.h"
 #include "EdGraph/EdGraphSchema.h"
 #include "BlueprintUtilities.h"
@@ -57,6 +59,29 @@ FPinDeletionQueue* PinDeletionQueue = new FPinDeletionQueue();;
 TArray<TPair<UEdGraphPin*, FString>> PinAllocationTracking;
 #endif //TRACK_PINS
 
+FArchive& operator<<(FArchive& Ar, FEdGraphTerminalType& T)
+{
+	Ar << T.TerminalCategory;
+	Ar << T.TerminalSubCategory;
+
+	// See: FArchive& operator<<( FArchive& Ar, FWeakObjectPtr& WeakObjectPtr )
+	// The PinSubCategoryObject should be serialized into the package.
+	if (!Ar.IsObjectReferenceCollector() || Ar.IsModifyingWeakAndStrongReferences() || Ar.IsPersistent())
+	{
+		UObject* Object = T.TerminalSubCategoryObject.Get(true);
+		Ar << Object;
+		if (Ar.IsLoading() || Ar.IsModifyingWeakAndStrongReferences())
+		{
+			T.TerminalSubCategoryObject = Object;
+		}
+	}
+
+	Ar << T.bTerminalIsConst;
+	Ar << T.bTerminalIsWeakPointer;
+
+	return Ar;
+}
+
 /////////////////////////////////////////////////////
 // FEdGraphPinType
 
@@ -82,6 +107,17 @@ bool FEdGraphPinType::Serialize(FArchive& Ar)
 		}
 	}
 
+	Ar.UsingCustomVersion(FBlueprintsObjectVersion::GUID);
+	if (Ar.CustomVer(FBlueprintsObjectVersion::GUID) >= FBlueprintsObjectVersion::AdvancedContainerSupport)
+	{
+		Ar << bIsMap;
+		if (bIsMap)
+		{
+			Ar << PinValueType;
+		}
+		Ar << bIsSet;
+	}
+
 	Ar << bIsArray;
 	Ar << bIsReference;
 	Ar << bIsWeakPointer;
@@ -958,6 +994,17 @@ bool UEdGraphPin::ImportTextItem(const TCHAR*& Buffer, int32 PortFlags, class UO
 	return true;
 }
 
+FEdGraphTerminalType UEdGraphPin::GetPrimaryTerminalType() const
+{
+	FEdGraphTerminalType TerminalType;
+	TerminalType.TerminalCategory = PinType.PinCategory;
+	TerminalType.TerminalSubCategory = PinType.PinSubCategory;
+	TerminalType.TerminalSubCategoryObject = PinType.PinSubCategoryObject;
+	TerminalType.bTerminalIsConst = PinType.bIsConst;
+	TerminalType.bTerminalIsWeakPointer = PinType.bIsWeakPointer;
+	return TerminalType;
+}
+
 void UEdGraphPin::MarkPendingKill()
 {
 	if (!bWasTrashed)
diff --git a/Engine/Source/Runtime/Engine/Public/InstancedReferenceSubobjectHelper.h b/Engine/Source/Runtime/Engine/Public/InstancedReferenceSubobjectHelper.h
--- a/Engine/Source/Runtime/Engine/Public/InstancedReferenceSubobjectHelper.h
+++ b/Engine/Source/Runtime/Engine/Public/InstancedReferenceSubobjectHelper.h
@@ -93,7 +93,7 @@ public:
 		}
 
 		const UObjectProperty* TargetPropety = Cast<UObjectProperty>(CurrentProp);
-		if ( ensure(TargetPropety && TargetPropety->HasAnyPropertyFlags(CPF_InstancedReference)) )
+		if (TargetPropety && TargetPropety->HasAnyPropertyFlags(CPF_InstancedReference))
 		{ 
 			return TargetPropety->GetObjectPropertyValue(ValuePtr);
 		}
diff --git a/Engine/Source/Runtime/PerfCounters/PerfCounters.Build.cs b/Engine/Source/Runtime/PerfCounters/PerfCounters.Build.cs
--- a/Engine/Source/Runtime/PerfCounters/PerfCounters.Build.cs
+++ b/Engine/Source/Runtime/PerfCounters/PerfCounters.Build.cs
@@ -15,7 +15,5 @@ public class PerfCounters : ModuleRules
 				"HTTP"
 			}
         );
-
-		Definitions.Add("WITH_PERFCOUNTERS=1");
 	}
 }
diff --git a/Engine/Source/Runtime/UMG/Private/UserWidget.cpp b/Engine/Source/Runtime/UMG/Private/UserWidget.cpp
--- a/Engine/Source/Runtime/UMG/Private/UserWidget.cpp
+++ b/Engine/Source/Runtime/UMG/Private/UserWidget.cpp
@@ -79,7 +79,7 @@ bool UUserWidget::Initialize()
 		}
 		else
 		{
-			CustomNativeInitilize();
+			InitializeNativeClassData();
 		}
 
 		if ( WidgetTree == nullptr )
diff --git a/Engine/Source/Runtime/UMG/Public/Blueprint/UserWidget.h b/Engine/Source/Runtime/UMG/Public/Blueprint/UserWidget.h
--- a/Engine/Source/Runtime/UMG/Public/Blueprint/UserWidget.h
+++ b/Engine/Source/Runtime/UMG/Public/Blueprint/UserWidget.h
@@ -145,8 +145,12 @@ public:
 	//~ End UObject Interface
 
 	virtual bool Initialize();
-	virtual void CustomNativeInitilize() {}
 
+protected:
+	/** THe function is implemented only in nativized widgets (automatically converted from BP to c++) */
+	virtual void InitializeNativeClassData() {}
+
+public:
 	//UVisual interface
 	virtual void ReleaseSlateResources(bool bReleaseChildren) override;
 	//~ End UVisual Interface
