Fixed anim timeline scrub ticks extending outside of their track when zoomed

diff --git a/Engine/Source/Editor/KismetWidgets/Private/SScrubControlPanel.cpp b/Engine/Source/Editor/KismetWidgets/Private/SScrubControlPanel.cpp
--- a/Engine/Source/Editor/KismetWidgets/Private/SScrubControlPanel.cpp
+++ b/Engine/Source/Editor/KismetWidgets/Private/SScrubControlPanel.cpp
@@ -53,6 +53,7 @@ void SScrubControlPanel::Construct( const SScrubControlPanel::FArguments& InArgs
 			SNew( SBorder )
 			[
 				SAssignNew(ScrubWidget, SScrubWidget)
+				.Clipping(EWidgetClipping::ClipToBounds)
 				.Value(InArgs._Value)
 				.NumOfKeys(InArgs._NumOfKeys)
 				.DisplayDrag(InArgs._DisplayDrag)
