Fix for Material graphs being in the simple verion permanently.

diff --git a/Engine/Source/Editor/KismetWidgets/Private/SLevelOfDetailBranchNode.cpp b/Engine/Source/Editor/KismetWidgets/Private/SLevelOfDetailBranchNode.cpp
--- a/Engine/Source/Editor/KismetWidgets/Private/SLevelOfDetailBranchNode.cpp
+++ b/Engine/Source/Editor/KismetWidgets/Private/SLevelOfDetailBranchNode.cpp
@@ -38,10 +38,19 @@ void SLevelOfDetailBranchNode::Tick( const FGeometry& AllottedGeometry, const do
 	{
 		LastCachedValue = CurrentValue;
 
-		ChildSlot
-		[
-			OnGetActiveDetailSlotContent.IsBound() ? OnGetActiveDetailSlotContent.Execute((CurrentValue != 0) ? false : true) : SNullWidget::NullWidget
-		];
-
+		if (OnGetActiveDetailSlotContent.IsBound())
+		{
+			ChildSlot
+			[
+				OnGetActiveDetailSlotContent.Execute((CurrentValue != 0) ? false : true)
+			];
+		}
+		else
+		{
+			ChildSlot
+			[
+				(CurrentValue != 0) ? ChildSlotLowDetail : ChildSlotHighDetail
+			];
+		}
 	}
 }
