Prevent the type of a variable from being changed in the blueprint diff window

diff --git a/Engine/Source/Editor/Kismet/Private/SBlueprintPalette.cpp b/Engine/Source/Editor/Kismet/Private/SBlueprintPalette.cpp
--- a/Engine/Source/Editor/Kismet/Private/SBlueprintPalette.cpp
+++ b/Engine/Source/Editor/Kismet/Private/SBlueprintPalette.cpp
@@ -702,7 +702,7 @@ public:
 			.TargetPinType(this, &SPinTypeSelectorHelper::OnGetVarType)
 			.OnPinTypeChanged(this, &SPinTypeSelectorHelper::OnVarTypeChanged)
 			.TypeTreeFilter(ETypeTreeFilter::None)
-			.bCompactSelector(true)
+			.SelectorType(BlueprintEditorPtr.IsValid() ? SPinTypeSelector::ESelectorType::Compact : SPinTypeSelector::ESelectorType::None)
 		];
 	}
 
diff --git a/Engine/Source/Editor/KismetWidgets/Private/SPinTypeSelector.cpp b/Engine/Source/Editor/KismetWidgets/Private/SPinTypeSelector.cpp
--- a/Engine/Source/Editor/KismetWidgets/Private/SPinTypeSelector.cpp
+++ b/Engine/Source/Editor/KismetWidgets/Private/SPinTypeSelector.cpp
@@ -260,21 +260,21 @@ void SPinTypeSelector::Construct(const FArguments& InArgs, FGetPinTypeTree GetPi
 	TreeViewHeight = InArgs._TreeViewHeight;
 
 	TargetPinType = InArgs._TargetPinType;
-	bIsCompactSelector = InArgs._bCompactSelector;
+	SelectorType = InArgs._SelectorType;
 
 	bIsRightMousePressed = false;
 
 	// Depending on if this is a compact selector or not, we generate a different compound widget
 	TSharedPtr<SWidget> Widget;
 
-	if (InArgs._bCompactSelector)
+	if (SelectorType == ESelectorType::Compact)
 	{
 		// Only have a combo button with an icon
 		Widget = SAssignNew( TypeComboButton, SComboButton )
 			.OnGetMenuContent(this, &SPinTypeSelector::GetMenuContent, false)
 			.ContentPadding(0)
 			.ToolTipText(this, &SPinTypeSelector::GetToolTipForComboBoxType)
-			.HasDownArrow(!InArgs._bCompactSelector)
+			.HasDownArrow(false)
 			.ButtonStyle(FEditorStyle::Get(),  "BlueprintEditor.CompactPinTypeSelector")
 			.ButtonContent()
 			[
@@ -287,7 +287,18 @@ void SPinTypeSelector::Construct(const FArguments& InArgs, FGetPinTypeTree GetPi
 				.ColorAndOpacity(this, &SPinTypeSelector::GetTypeIconColor)
 			];
 	}
-	else
+	else if (SelectorType == ESelectorType::None)
+	{
+		Widget = SNew(
+					SDoubleImage,
+					TAttribute<const FSlateBrush*>(this, &SPinTypeSelector::GetSecondaryTypeIconImage),
+					TAttribute<FSlateColor>(this, &SPinTypeSelector::GetSecondaryTypeIconColor)
+				)
+				.ToolTipText(this, &SPinTypeSelector::GetToolTipForComboBoxType)
+				.Image(this, &SPinTypeSelector::GetTypeIconImage)
+				.ColorAndOpacity(this, &SPinTypeSelector::GetTypeIconColor);
+	}
+	else if (SelectorType == ESelectorType::Full)
 	{
 		// Traditional Pin Type Selector with a combo button, the icon, the current type name, and a toggle button for being an array
 		TSharedPtr<SWidget> ContainerControl = SNew(SComboButton)
@@ -1064,21 +1075,21 @@ FText SPinTypeSelector::GetToolTipForComboBoxType() const
 	FText EditText;
 	if(IsEnabled())
 	{
-		if (bIsCompactSelector)
+		if (SelectorType == ESelectorType::Compact)
 		{
-			EditText = LOCTEXT("CompactPinTypeSelector", "Left click to select the variable's pin type. Right click to toggle the type as an array.");
+			EditText = LOCTEXT("CompactPinTypeSelector", "Left click to select the variable's pin type. Right click to toggle the type as an array.\n");
 		}
-		else
+		else if (SelectorType == ESelectorType::Full)
 		{
-			EditText = LOCTEXT("PinTypeSelector", "Select the variable's pin type.");
+			EditText = LOCTEXT("PinTypeSelector", "Select the variable's pin type.\n");
 		}
 	}
 	else
 	{
-		EditText = LOCTEXT("PinTypeSelector_Disabled", "Cannot edit variable type when they are inherited from parent.");
+		EditText = LOCTEXT("PinTypeSelector_Disabled", "Cannot edit variable type when they are inherited from parent.\n");
 	}
 
-	return FText::Format(LOCTEXT("PrimaryTypeTwoLines", "{0}\nCurrent Type: {1}"), EditText, GetTypeDescription());
+	return FText::Format(LOCTEXT("PrimaryTypeTwoLines", "{0}Current Type: {1}"), EditText, GetTypeDescription());
 }
 
 FText SPinTypeSelector::GetToolTipForComboBoxSecondaryType() const
@@ -1149,7 +1160,7 @@ FText SPinTypeSelector::GetToolTipForContainerWidget() const
 
 FReply SPinTypeSelector::OnMouseButtonDown( const FGeometry& MyGeometry, const FPointerEvent& MouseEvent )
 {
-	if(bIsCompactSelector && MouseEvent.GetEffectingButton() == EKeys::RightMouseButton)
+	if (SelectorType == ESelectorType::Compact && MouseEvent.GetEffectingButton() == EKeys::RightMouseButton)
 	{
 		bIsRightMousePressed = true;
 		return FReply::Handled();
@@ -1160,7 +1171,7 @@ FReply SPinTypeSelector::OnMouseButtonDown( const FGeometry& MyGeometry, const F
 
 FReply SPinTypeSelector::OnMouseButtonUp( const FGeometry& MyGeometry, const FPointerEvent& MouseEvent )
 {
-	if(bIsCompactSelector && MouseEvent.GetEffectingButton() == EKeys::RightMouseButton)
+	if (SelectorType == ESelectorType::Compact && MouseEvent.GetEffectingButton() == EKeys::RightMouseButton)
 	{
 		if (bIsRightMousePressed)
 		{
diff --git a/Engine/Source/Editor/KismetWidgets/Public/SPinTypeSelector.h b/Engine/Source/Editor/KismetWidgets/Public/SPinTypeSelector.h
--- a/Engine/Source/Editor/KismetWidgets/Public/SPinTypeSelector.h
+++ b/Engine/Source/Editor/KismetWidgets/Public/SPinTypeSelector.h
@@ -46,6 +46,14 @@ public:
 	static TSharedRef<SWidget> ConstructPinTypeImage(TAttribute<const FSlateBrush*> PrimaryIcon, TAttribute<FSlateColor> PrimaryColor, TAttribute<const FSlateBrush*> SecondaryIcon, TAttribute<FSlateColor> SecondaryColor );
 	static TSharedRef<SWidget> ConstructPinTypeImage(UEdGraphPin* Pin);
 
+	/** Which type of selector should be used: compact or full mode, or not a selector at all, but just the type image. */
+	enum class ESelectorType : uint8
+	{
+		None,
+		Compact,
+		Full
+	};
+
 	SLATE_BEGIN_ARGS( SPinTypeSelector )
 		: _TargetPinType()
 		, _Schema(nullptr)
@@ -54,7 +62,7 @@ public:
 		, _TreeViewWidth(300.f)
 		, _TreeViewHeight(400.f)
 		, _Font( FEditorStyle::GetFontStyle( TEXT("NormalFont") ) )
-		, _bCompactSelector( false )
+		, _SelectorType( ESelectorType::Full )
 		{}
 		SLATE_ATTRIBUTE( FEdGraphPinType, TargetPinType )
 		SLATE_ARGUMENT( const UEdGraphSchema_K2*, Schema )
@@ -65,7 +73,7 @@ public:
 		SLATE_EVENT( FOnPinTypeChanged, OnPinTypePreChanged )
 		SLATE_EVENT( FOnPinTypeChanged, OnPinTypeChanged )
 		SLATE_ATTRIBUTE( FSlateFontInfo, Font )
-		SLATE_ARGUMENT( bool, bCompactSelector )
+		SLATE_ARGUMENT( ESelectorType, SelectorType )
 	SLATE_END_ARGS()
 public:
 	void Construct(const FArguments& InArgs, FGetPinTypeTree GetPinTypeTreeFunc);
@@ -125,8 +133,8 @@ protected:
 	/** TRUE when the right mouse button is pressed, keeps from handling a right click that does not begin in the widget */
 	bool bIsRightMousePressed;
 
-	/** TRUE if displaying a compact selector widget, some functionality is enabled in different ways if this is TRUE */
-	bool bIsCompactSelector;
+	/** Whether the selector is using the compact or full mode, or not a selector at all, but just the type image.*/
+	ESelectorType SelectorType;
 
 	/** Holds a cache of the allowed Object Reference types for the last sub-menu opened. */
 	TArray<FObjectReferenceListItem> AllowedObjectReferenceTypes;
