PR #6057: Clip the value type selector of maps (Contributed by [REDACTED_PERSON])

diff --git a/Engine/Source/Editor/KismetWidgets/Private/SPinTypeSelector.cpp b/Engine/Source/Editor/KismetWidgets/Private/SPinTypeSelector.cpp
--- a/Engine/Source/Editor/KismetWidgets/Private/SPinTypeSelector.cpp
+++ b/Engine/Source/Editor/KismetWidgets/Private/SPinTypeSelector.cpp
@@ -434,6 +434,8 @@ void SPinTypeSelector::Construct(const FArguments& InArgs, FGetPinTypeTree GetPi
 				.ButtonContent()
 				[
 					SNew(SHorizontalBox)
+					.Clipping(EWidgetClipping::OnDemand)
+
 					+SHorizontalBox::Slot()
 					.AutoWidth()
 					.VAlign(VAlign_Center)
