Fixed: Map value text expands outside of the drop down menu

diff --git a/Engine/Source/Editor/KismetWidgets/Private/SPinTypeSelector.cpp b/Engine/Source/Editor/KismetWidgets/Private/SPinTypeSelector.cpp
--- a/Engine/Source/Editor/KismetWidgets/Private/SPinTypeSelector.cpp
+++ b/Engine/Source/Editor/KismetWidgets/Private/SPinTypeSelector.cpp
@@ -433,7 +433,6 @@ void SPinTypeSelector::Construct(const FArguments& InArgs, FGetPinTypeTree GetPi
 						.ColorAndOpacity( this, &SPinTypeSelector::GetSecondaryTypeIconColor )
 					]
 					+SHorizontalBox::Slot()
-					.AutoWidth()
 					.VAlign(VAlign_Center)
 					.HAlign(HAlign_Left)
 					[
