Removed animation and extended duration for a notification toast

diff --git a/Engine/Source/Editor/KismetWidgets/Private/SPinTypeSelector.cpp b/Engine/Source/Editor/KismetWidgets/Private/SPinTypeSelector.cpp
--- a/Engine/Source/Editor/KismetWidgets/Private/SPinTypeSelector.cpp
+++ b/Engine/Source/Editor/KismetWidgets/Private/SPinTypeSelector.cpp
@@ -822,7 +822,9 @@ void SPinTypeSelector::OnSelectPinType(FPinTypeTreeItem InItem, FName InPinCateg
 		// inform user via toast why the type change was exceptional and clear IsMap/IsSetness because this type cannot be hashed:
 		const FText NotificationText = FText::Format(LOCTEXT("TypeCannotBeHashed", "Container type cleared because '{0}' does not have a GetTypeHash function. Maps and Sets require a hash function to insert and find elements"), UEdGraphSchema_K2::TypeToText(NewTargetPinType));
 		FNotificationInfo Info(NotificationText);
-		Info.ExpireDuration = 8.0f;
+		Info.FadeInDuration = 0.0f;
+		Info.FadeOutDuration = 0.0f;
+		Info.ExpireDuration = 10.0f;
 		FSlateNotificationManager::Get().AddNotification(Info);
 	}
 
