Move to a CreateBlueprintFromMode enum instead of boolean

diff --git a/Engine/Source/Editor/Kismet/Private/BlueprintEditorViewportContextMenuExtender.cpp b/Engine/Source/Editor/Kismet/Private/BlueprintEditorViewportContextMenuExtender.cpp
--- a/Engine/Source/Editor/Kismet/Private/BlueprintEditorViewportContextMenuExtender.cpp
+++ b/Engine/Source/Editor/Kismet/Private/BlueprintEditorViewportContextMenuExtender.cpp
@@ -199,7 +199,7 @@ void FillBlueprintOptions(FMenuBuilder& MenuBuilder, TArray<AActor*> SelectedAct
 		if(bCanHarvestComponentsForBlueprint)
 		{
 			AActor* ActorOverride = nullptr;
-			FUIAction CreateBlueprintAction( FExecuteAction::CreateStatic( &FCreateBlueprintFromActorDialog::OpenDialog, true, ActorOverride ) );
+			FUIAction CreateBlueprintAction( FExecuteAction::CreateStatic( &FCreateBlueprintFromActorDialog::OpenDialog, ECreateBlueprintFromActorMode::Harvest, ActorOverride ) );
 			MenuBuilder.AddMenuEntry(LOCTEXT("CreateBlueprint", "Create Blueprint..."), LOCTEXT("CreateBlueprint_Tooltip", "Harvest Components from Selected Actors and create Blueprint"), FSlateIcon(FEditorStyle::GetStyleSetName(), "Kismet.HarvestBlueprintFromActors"), CreateBlueprintAction);
 		}
 	}
diff --git a/Engine/Source/Editor/Kismet/Private/SSCSEditor.cpp b/Engine/Source/Editor/Kismet/Private/SSCSEditor.cpp
--- a/Engine/Source/Editor/Kismet/Private/SSCSEditor.cpp
+++ b/Engine/Source/Editor/Kismet/Private/SSCSEditor.cpp
@@ -6682,8 +6682,7 @@ void SSCSEditor::OnResetToBlueprintDefaults() const
 
 void SSCSEditor::PromoteToBlueprint() const
 {
-	bool bHarvest = false;
-	FCreateBlueprintFromActorDialog::OpenDialog(bHarvest, GetActorContext());
+	FCreateBlueprintFromActorDialog::OpenDialog(ECreateBlueprintFromActorMode::Subclass, GetActorContext());
 }
 
 FReply SSCSEditor::OnPromoteToBlueprintClicked()
diff --git a/Engine/Source/Editor/KismetWidgets/Private/CreateBlueprintFromActorDialog.cpp b/Engine/Source/Editor/KismetWidgets/Private/CreateBlueprintFromActorDialog.cpp
--- a/Engine/Source/Editor/KismetWidgets/Private/CreateBlueprintFromActorDialog.cpp
+++ b/Engine/Source/Editor/KismetWidgets/Private/CreateBlueprintFromActorDialog.cpp
@@ -19,7 +19,7 @@
 
 TWeakObjectPtr<AActor> FCreateBlueprintFromActorDialog::ActorOverride = nullptr;
 
-void FCreateBlueprintFromActorDialog::OpenDialog(bool bInHarvest, AActor* InActorOverride )
+void FCreateBlueprintFromActorDialog::OpenDialog(ECreateBlueprintFromActorMode CreateMode, AActor* InActorOverride )
 {
 	ActorOverride = InActorOverride;
 
@@ -36,7 +36,7 @@ void FCreateBlueprintFromActorDialog::OpenDialog(bool bInHarvest, AActor* InActo
 		.AssetFilenameSuffix(TEXT("Blueprint"))
 		.HeadingText(LOCTEXT("CreateBlueprintFromActor_Heading", "Blueprint Name"))
 		.CreateButtonText(LOCTEXT("CreateBlueprintFromActor_ButtonLabel", "Create Blueprint"))
-		.OnCreateAssetAction(FOnPathChosen::CreateStatic(FCreateBlueprintFromActorDialog::OnCreateBlueprint, bInHarvest))
+		.OnCreateAssetAction(FOnPathChosen::CreateStatic(FCreateBlueprintFromActorDialog::OnCreateBlueprint, CreateMode))
 	);
 
 	TSharedPtr<SWindow> RootWindow = FGlobalTabmanager::Get()->GetRootWindow();
@@ -50,40 +50,47 @@ void FCreateBlueprintFromActorDialog::OpenDialog(bool bInHarvest, AActor* InActo
 	}
 }
 
-void FCreateBlueprintFromActorDialog::OnCreateBlueprint(const FString& InAssetPath, bool bInHarvest)
+void FCreateBlueprintFromActorDialog::OnCreateBlueprint(const FString& InAssetPath, ECreateBlueprintFromActorMode CreateMode)
 {
 	UBlueprint* Blueprint = NULL;
 
-	if(bInHarvest)
+	switch (CreateMode) 
 	{
-		TArray<AActor*> Actors;
-
-		USelection* SelectedActors = GEditor->GetSelectedActors();
-		for(FSelectionIterator Iter(*SelectedActors); Iter; ++Iter)
+		case ECreateBlueprintFromActorMode::Harvest:
 		{
-			// We only care about actors that are referenced in the world for literals, and also in the same level as this blueprint
-			AActor* Actor = Cast<AActor>(*Iter);
-			if(Actor)
+			TArray<AActor*> Actors;
+
+			USelection* SelectedActors = GEditor->GetSelectedActors();
+			for(FSelectionIterator Iter(*SelectedActors); Iter; ++Iter)
 			{
-				Actors.Add(Actor);
+				// We only care about actors that are referenced in the world for literals, and also in the same level as this blueprint
+				if (AActor* Actor = Cast<AActor>(*Iter))
+				{
+					Actors.Add(Actor);
+				}
 			}
+
+			const bool bReplaceActor = true;
+			Blueprint = FKismetEditorUtilities::HarvestBlueprintFromActors(InAssetPath, Actors, bReplaceActor);
 		}
-		Blueprint = FKismetEditorUtilities::HarvestBlueprintFromActors(InAssetPath, Actors, true);
-	}
-	else
-	{
-		AActor* ActorToUse = ActorOverride.Get();
+		break;
 
-		if( !ActorToUse )
+		case ECreateBlueprintFromActorMode::Subclass:
 		{
-			TArray< UObject* > SelectedActors;
-			GEditor->GetSelectedActors()->GetSelectedObjects(AActor::StaticClass(), SelectedActors);
-			check(SelectedActors.Num());
-			ActorToUse =  Cast<AActor>(SelectedActors[0]);
-		}
+			AActor* ActorToUse = ActorOverride.Get();
 
-		const bool bReplaceActor = true;
-		Blueprint = FKismetEditorUtilities::CreateBlueprintFromActor(InAssetPath, ActorToUse, bReplaceActor);
+			if (!ActorToUse)
+			{
+				TArray< UObject* > SelectedActors;
+				GEditor->GetSelectedActors()->GetSelectedObjects(AActor::StaticClass(), SelectedActors);
+				check(SelectedActors.Num() == 1);
+				ActorToUse = Cast<AActor>(SelectedActors[0]);
+			}
+
+			const bool bReplaceActor = true;
+			Blueprint = FKismetEditorUtilities::CreateBlueprintFromActor(InAssetPath, ActorToUse, bReplaceActor);
+		}
+		break;
 	}
 
 	if(Blueprint)
diff --git a/Engine/Source/Editor/KismetWidgets/Public/CreateBlueprintFromActorDialog.h b/Engine/Source/Editor/KismetWidgets/Public/CreateBlueprintFromActorDialog.h
--- a/Engine/Source/Editor/KismetWidgets/Public/CreateBlueprintFromActorDialog.h
+++ b/Engine/Source/Editor/KismetWidgets/Public/CreateBlueprintFromActorDialog.h
@@ -8,25 +8,42 @@
 //////////////////////////////////////////////////////////////////////////
 // FCreateBlueprintFromActorDialog
 
+/** Enum defining which mode to use when creating the blueprint from the selected actors */
+enum class ECreateBlueprintFromActorMode : uint8
+{
+	None,        // Indicates that a user is unable to create blueprints from the currently selected actor set
+	Harvest,     // Harvest all the components of the selected actors and create an actor blueprint with those components in it.
+	Subclass     // Create a subclass of the selected Actor's class with defaults from the selected Actor. Valid only when there is a single selected actor.
+};
+
 class FCreateBlueprintFromActorDialog
 {
 public:
+
 	/** 
 	 * Static function to access constructing this window.
 	 *
-	 * @param bInHarvest		true if the components of the selected actors should be harvested for the blueprint.
+	 * @param CreateMode		The mode to use when creating a blueprint from the selected actors
 	 * @param ActorOverride		If set convert the specified actor, if null use the currently selected actor
 	 */
-	static KISMETWIDGETS_API void OpenDialog(bool bInHarvest, AActor* InActorOverride = nullptr);
+	static KISMETWIDGETS_API void OpenDialog(ECreateBlueprintFromActorMode CreateMode, AActor* InActorOverride = nullptr);
+
+	UE_DEPRECATED(4.25, "Use version of OpenDialog that takes the CreateMode enum")
+	static KISMETWIDGETS_API void OpenDialog(bool bInHarvest, AActor* InActorOverride = nullptr)
+	{
+		OpenDialog(bInHarvest ? ECreateBlueprintFromActorMode::Harvest : ECreateBlueprintFromActorMode::Subclass, InActorOverride);
+	}
+
+
 private:
 
 	/** 
-	 * Will create the blueprint
+	 * Create the blueprint in response to the path being specified by the user
 	 *
 	 * @param InAssetPath		Path of the asset to create
-	 * @param bInHarvest		true if the components of the selected actors should be harvested for the blueprint.
+	 * @param bInHarvest		The mode to use when creating a blueprint from the selected actors
 	 */
-	static void OnCreateBlueprint(const FString& InAssetPath, bool bInHarvest);
-private:
+	static void OnCreateBlueprint(const FString& InAssetPath, ECreateBlueprintFromActorMode CreateMode);
+
 	static TWeakObjectPtr<AActor> ActorOverride;
 };
diff --git a/Engine/Source/Editor/LevelEditor/Private/LevelEditorActions.cpp b/Engine/Source/Editor/LevelEditor/Private/LevelEditorActions.cpp
--- a/Engine/Source/Editor/LevelEditor/Private/LevelEditorActions.cpp
+++ b/Engine/Source/Editor/LevelEditor/Private/LevelEditorActions.cpp
@@ -2333,8 +2333,7 @@ bool FLevelEditorActionCallbacks::CanHarvestSelectedActorsIntoBlueprintClass()
 
 void FLevelEditorActionCallbacks::HarvestSelectedActorsIntoBlueprintClass()
 {
-	const bool bHarvest = true;
-	FCreateBlueprintFromActorDialog::OpenDialog(bHarvest);
+	FCreateBlueprintFromActorDialog::OpenDialog(ECreateBlueprintFromActorMode::Harvest);
 }
 
 bool FLevelEditorActionCallbacks::CanSubclassSelectedActorIntoBlueprintClass()
@@ -2350,8 +2349,7 @@ bool FLevelEditorActionCallbacks::CanSubclassSelectedActorIntoBlueprintClass()
 
 void FLevelEditorActionCallbacks::SubclassSelectedActorIntoBlueprintClass()
 {
-	const bool bHarvest = false;
-	FCreateBlueprintFromActorDialog::OpenDialog(bHarvest);
+	FCreateBlueprintFromActorDialog::OpenDialog(ECreateBlueprintFromActorMode::Subclass);
 }
 
 void FLevelEditorActionCallbacks::CheckOutProjectSettingsConfig( )
