Create blueprint from selection now brings up a dialog that allows the creation method and parent class to be selected.

diff --git a/Engine/Source/Editor/KismetWidgets/Private/CreateBlueprintFromActorDialog.cpp b/Engine/Source/Editor/KismetWidgets/Private/CreateBlueprintFromActorDialog.cpp
--- a/Engine/Source/Editor/KismetWidgets/Private/CreateBlueprintFromActorDialog.cpp
+++ b/Engine/Source/Editor/KismetWidgets/Private/CreateBlueprintFromActorDialog.cpp
@@ -2,57 +2,675 @@
 
 #include "CreateBlueprintFromActorDialog.h"
 #include "GameFramework/Actor.h"
+#include "Misc/FileHelper.h"
+#include "Misc/MessageDialog.h"
 #include "Misc/PackageName.h"
 #include "Widgets/DeclarativeSyntaxSupport.h"
 #include "Widgets/SWindow.h"
+#include "Widgets/Images/SImage.h"
+#include "Widgets/Input/SButton.h"
+#include "Widgets/Input/SEditableTextBox.h"
+#include "Widgets/Layout/SExpandableArea.h"
+#include "Widgets/Layout/SGridPanel.h"
+#include "Widgets/Layout/SUniformGridPanel.h"
 #include "Framework/Application/SlateApplication.h"
 #include "Framework/Docking/TabManager.h"
 #include "Editor/EditorEngine.h"
 #include "Engine/Selection.h"
 #include "Editor.h"
-#include "SCreateAssetFromObject.h"
 #include "Kismet2/KismetEditorUtilities.h"
 #include "Framework/Notifications/NotificationManager.h"
 #include "Widgets/Notifications/SNotificationList.h"
+#include "AssetRegistryModule.h"
+#include "AssetToolsModule.h"
+#include "ClassViewerFilter.h"
+#include "ClassViewerModule.h"
+#include "SClassViewer.h"
+#include "ContentBrowserModule.h"
+#include "IContentBrowserSingleton.h"
+#include "PackageTools.h"
+#include "DetailLayoutBuilder.h"
 
 #define LOCTEXT_NAMESPACE "CreateBlueprintFromActorDialog"
 
 TWeakObjectPtr<AActor> FCreateBlueprintFromActorDialog::ActorOverride = nullptr;
 
+class SSCreateBlueprintPicker : public SCompoundWidget
+{
+public:
+	SLATE_BEGIN_ARGS(SSCreateBlueprintPicker)
+	{}
+
+	SLATE_ARGUMENT(TSharedPtr<SWindow>, ParentWindow)
+		SLATE_ARGUMENT(FClassViewerInitializationOptions, Options)
+		SLATE_ARGUMENT(ECreateBlueprintFromActorMode, CreateMode)
+		SLATE_END_ARGS()
+
+	/** Constructs this widget with InArgs */
+	void Construct(const FArguments& InArgs);
+
+	/** Handler for when a class is picked in the class picker */
+	void OnClassPicked(UClass* InChosenClass);
+
+	/** Handler for when "Ok" we selected in the class viewer */
+	FReply OnClassPickerConfirmed();
+
+	/** Handler for when "Cancel" we selected in the class viewer */
+	FReply OnClassPickerCanceled();
+
+	/** Handler for when "..." is clicked to pick an asset path */
+	FReply OnPathPickerSummoned();
+
+	/** Handler for the custom button to hide/unhide the class viewer */
+	void OnCustomAreaExpansionChanged(bool bExpanded);
+
+	/** Callback when the user changes the filename for the Blueprint */
+	void OnFilenameChanged(const FText& InNewName);
+
+	/** Common function to evaluate whether the asset name given the asset path context, is valid */
+	void UpdateFilenameStatus();
+
+	/** select button visibility delegate */
+	EVisibility GetSelectButtonVisibility() const;
+
+	ECheckBoxState IsCreateModeChecked(ECreateBlueprintFromActorMode InCreateMode) const
+	{
+		return (CreateMode == InCreateMode ? ECheckBoxState::Checked : ECheckBoxState::Unchecked);
+	};
+
+	void OnCreateModeChanged(ECheckBoxState NewCheckedState, ECreateBlueprintFromActorMode InCreateMode) 
+	{ 
+		if (NewCheckedState == ECheckBoxState::Checked) 
+		{ 
+			CreateMode = InCreateMode; 
+		} 	
+	};
+
+	FSlateColor GetCreateModeTextColor(ECreateBlueprintFromActorMode InCreateMode) const
+	{ 
+		return (CreateMode == InCreateMode ? FSlateColor(FLinearColor(0, 0, 0)) : FSlateColor(FLinearColor(0.72f, 0.72f, 0.72f, 1.f))); 
+	};
+
+	/** Overridden from SWidget: Called when a key is pressed down - capturing copy */
+	FReply OnKeyDown(const FGeometry& MyGeometry, const FKeyEvent& InKeyEvent);
+
+	/** A pointer to the window that is asking the user to select a parent class */
+	TWeakPtr<SWindow> WeakParentWindow;
+
+	/** A pointer to a class viewer **/
+	TSharedPtr<SClassViewer> ClassViewer;
+
+	/** Filename textbox widget */
+	TSharedPtr<SEditableTextBox> FileNameWidget;
+
+	/** The class that was last clicked on */
+	UClass* ChosenClass;
+
+	/** The path the asset should be created at */
+	FString AssetPath;
+
+	/** The the name for the new asset */
+	FString AssetName;
+
+	/** The method to use when creating the actor */
+	ECreateBlueprintFromActorMode CreateMode;
+
+	/** A flag indicating that Ok was selected */
+	bool bPressedOk;
+
+	/** A flag indicating the current asset name is invalid */
+	bool bIsReportingError;
+};
+
+void SSCreateBlueprintPicker::Construct(const FArguments& InArgs)
+{
+	WeakParentWindow = InArgs._ParentWindow;
+
+	bPressedOk = false;
+	ChosenClass = nullptr;
+	CreateMode = InArgs._CreateMode;
+
+	FContentBrowserModule& ContentBrowserModule = FModuleManager::LoadModuleChecked<FContentBrowserModule>("ContentBrowser");
+	FClassViewerModule& ClassViewerModule = FModuleManager::LoadModuleChecked<FClassViewerModule>("ClassViewer");
+	FAssetToolsModule& AssetToolsModule = FModuleManager::LoadModuleChecked<FAssetToolsModule>("AssetTools");
+
+	ClassViewer = StaticCastSharedRef<SClassViewer>(ClassViewerModule.CreateClassViewer(InArgs._Options, FOnClassPicked::CreateSP(this, &SSCreateBlueprintPicker::OnClassPicked)));
+
+	FString PackageName;
+	AssetPath = ContentBrowserModule.Get().GetCurrentPath();
+
+	int32 NumSelectedActors = 0;
+
+	bool bCanHarvestComponents = false;
+	bool bCanSubclass = true;
+	bool bCanCreatePrefab = true;
+
+	for (FSelectionIterator Iter(*GEditor->GetSelectedActors()); Iter; ++Iter)
+	{
+		AActor* Actor = Cast<AActor>(*Iter);
+		if (Actor)
+		{
+			if (NumSelectedActors == 0)
+			{
+				AssetName += Actor->GetActorLabel();
+				AssetName += TEXT("_");
+
+				bCanSubclass = FKismetEditorUtilities::CanCreateBlueprintOfClass(Actor->GetClass());
+			}
+
+			if (Actor->GetClass()->HasAnyClassFlags(CLASS_NotPlaceable))
+			{
+				bCanCreatePrefab = false;
+			}
+
+			bCanHarvestComponents = true;
+
+			++NumSelectedActors;
+		}
+	}
+
+	bCanSubclass = bCanSubclass && NumSelectedActors == 1;
+
+	AssetName = UPackageTools::SanitizePackageName(AssetName + TEXT("Blueprint"));
+
+	FString BasePath = AssetPath / AssetName;
+	AssetToolsModule.Get().CreateUniqueAssetName(BasePath, TEXT(""), PackageName, AssetName);
+
+	TSharedPtr<SGridPanel> CreationMethodSection;
+
+	struct FCreateModeDetails
+	{
+		FText Label;
+		FText Description;
+		ECreateBlueprintFromActorMode CreateMode;
+		bool bEnabled;
+	};
+
+	const FCreateModeDetails CreateModeDetails[3] =
+	{
+		{ LOCTEXT("CreateMode_Subclass", "New Subclass"), LOCTEXT("CreateMode_Subclass_Description", "Replace the selected actor with an instance of a new Blueprint Class inherited from the selected parent class."), ECreateBlueprintFromActorMode::Subclass, bCanSubclass },
+		{ LOCTEXT("CreateMode_ChildActor", "Child Actors"), LOCTEXT("CreateMode_ChildActor_Description", "Replace the selected actors with an instance of a new Blueprint Class inherited from the selected parent class with each of the selected Actors as a Child Actor."), ECreateBlueprintFromActorMode::ChildActor, bCanCreatePrefab },
+		{ LOCTEXT("CreateMode_Harvest", "Harvest Components"), LOCTEXT("CreateMode_Harvest_Description", "Replace the selected actors with an instance of a new Blueprint Class inherited from the selected parent class that contains the components."), ECreateBlueprintFromActorMode::Harvest, bCanHarvestComponents }
+	};
+
+	const FCheckBoxStyle& RadioStyle = FEditorStyle::Get().GetWidgetStyle<FCheckBoxStyle>("Property.ToggleButton");
+
+	SAssignNew(CreationMethodSection, SGridPanel)
+	.FillColumn(1, 1.f);
+
+	for (int32 Index = 0; Index < UE_ARRAY_COUNT(CreateModeDetails); ++Index)
+	{
+		CreationMethodSection->AddSlot(0, Index)
+		.Padding(10.0f, 5.0f, 5.0f, 5.0f)
+		.VAlign(VAlign_Center)
+		[
+			SNew(SCheckBox)
+			.Style(&RadioStyle)
+			.IsEnabled(CreateModeDetails[Index].bEnabled)
+			.IsChecked_Raw(this, &SSCreateBlueprintPicker::IsCreateModeChecked, CreateModeDetails[Index].CreateMode)
+			.OnCheckStateChanged(this, &SSCreateBlueprintPicker::OnCreateModeChanged, CreateModeDetails[Index].CreateMode)
+			[
+				SNew(SHorizontalBox)
+				+ SHorizontalBox::Slot()
+				.FillWidth(1.0f)
+				.VAlign(VAlign_Center)
+				.HAlign(HAlign_Center)
+				.Padding(6, 2)
+				[
+					SNew(STextBlock)
+					.Text(CreateModeDetails[Index].Label)
+					.Font(IDetailLayoutBuilder::GetDetailFont())
+					.ColorAndOpacity(this, &SSCreateBlueprintPicker::GetCreateModeTextColor, CreateModeDetails[Index].CreateMode)
+				]
+			]
+		];
+
+		CreationMethodSection->AddSlot(1, Index)
+		.Padding(1.0f, 5.0f, 1.0f, 5.0f)
+		.VAlign(VAlign_Center)
+		[
+			SNew(STextBlock)
+			.Text(CreateModeDetails[Index].Description)
+			.IsEnabled(CreateModeDetails[Index].bEnabled)
+			.AutoWrapText(true)
+		];
+	}
+
+	ChildSlot
+	[
+		SNew(SBorder)
+		.Visibility(EVisibility::Visible)
+		.BorderImage(FEditorStyle::GetBrush("Menu.Background"))
+		[
+			SNew(SBox)
+			.Visibility(EVisibility::Visible)
+			[
+				SNew(SVerticalBox)
+				+SVerticalBox::Slot()
+				.AutoHeight()
+				.Padding(0.0f, 10.0f, 0.0f, 0.0f)
+				[
+					SNew(SGridPanel)
+					.FillColumn(1, 1.f)
+					+SGridPanel::Slot(0, 0)
+					.Padding(0.0f, 0.0f, 5.0f, 2.0f)
+					.VAlign(VAlign_Center)
+					[
+						SNew(STextBlock)
+						.Text(LOCTEXT("CreateBlueprintFromActor_NameLabel", "Blueprint Name"))
+					]
+					+SGridPanel::Slot(1, 0)
+					.Padding(0.0f, 0.0f, 0.0f, 5.0f)
+					[
+						SAssignNew(FileNameWidget, SEditableTextBox)
+						.Text(FText::FromString(AssetName))
+						.OnTextChanged(this, &SSCreateBlueprintPicker::OnFilenameChanged)
+					]
+					+SGridPanel::Slot(0, 1)
+					.Padding(0.0f, 0.0f, 5.0f, 0.0f)
+					.VAlign(VAlign_Center)
+					[
+						SNew(STextBlock)
+						.Text(LOCTEXT("CreateBlueprintFromActor_PathLabel", "Path"))
+					]
+					+SGridPanel::Slot(1, 1)
+					[
+						SNew(SHorizontalBox)
+						+SHorizontalBox::Slot()
+						.FillWidth(1.f)
+						[
+							SNew(SEditableTextBox)
+							.Text(TAttribute<FText>::Create(TAttribute<FText>::FGetter::CreateLambda([this]() { return FText::FromString(AssetPath); })))
+							.IsReadOnly(true)
+						]
+						+SHorizontalBox::Slot()
+						.AutoWidth()
+						[
+							SNew(SButton)
+							.Text(FText::FromString(TEXT("...")))
+							.OnClicked(this, &SSCreateBlueprintPicker::OnPathPickerSummoned)
+						]
+					]
+				]
+				+SVerticalBox::Slot()
+				.AutoHeight()
+				.Padding(0.0f, 10.0f, 0.0f, 0.0f)
+				[
+					SNew(SExpandableArea)
+					.AreaTitle(LOCTEXT("CreationMethod", "Creation Method"))
+					.BodyContent()
+					[
+						CreationMethodSection.ToSharedRef()
+					]
+				]
+				+SVerticalBox::Slot()
+				.FillHeight(1.f)
+				.Padding(0.0f, 10.0f, 0.0f, 0.0f)
+				[
+					SNew(SExpandableArea)
+					.MaxHeight(320.f)
+					.InitiallyCollapsed(false)
+					.AreaTitle(NSLOCTEXT("SClassPickerDialog", "AllClassesAreaTitle", "Parent Class"))
+					.OnAreaExpansionChanged(this, &SSCreateBlueprintPicker::OnCustomAreaExpansionChanged)
+					.BodyContent()
+					[
+						ClassViewer.ToSharedRef()
+					]
+				]
+				+SVerticalBox::Slot()
+				.AutoHeight()
+				.HAlign(HAlign_Right)
+				.VAlign(VAlign_Bottom)
+				.Padding(8)
+				[
+					SNew(SUniformGridPanel)
+					.SlotPadding(FEditorStyle::GetMargin("StandardDialog.SlotPadding"))
+					+SUniformGridPanel::Slot(0,0)
+					[
+						SNew(SButton)
+						.Text(NSLOCTEXT("SClassPickerDialog", "ClassPickerSelectButton", "Select"))
+						.HAlign(HAlign_Center)
+						.Visibility( this, &SSCreateBlueprintPicker::GetSelectButtonVisibility )
+						.ContentPadding(FEditorStyle::GetMargin("StandardDialog.ContentPadding"))
+						.OnClicked(this, &SSCreateBlueprintPicker::OnClassPickerConfirmed)
+						.ButtonStyle(FEditorStyle::Get(), "FlatButton.Success")
+						.TextStyle(FEditorStyle::Get(), "FlatButton.DefaultTextStyle")
+					]
+					+SUniformGridPanel::Slot(1,0)
+					[
+						SNew(SButton)
+						.Text(NSLOCTEXT("SClassPickerDialog", "ClassPickerCancelButton", "Cancel"))
+						.HAlign(HAlign_Center)
+						.ContentPadding(FEditorStyle::GetMargin("StandardDialog.ContentPadding"))
+						.OnClicked(this, &SSCreateBlueprintPicker::OnClassPickerCanceled)
+						.ButtonStyle(FEditorStyle::Get(), "FlatButton.Default")
+						.TextStyle(FEditorStyle::Get(), "FlatButton.DefaultTextStyle")
+					]
+				]
+			]
+		]
+	];
+
+	if (WeakParentWindow.IsValid())
+	{
+		WeakParentWindow.Pin().Get()->SetWidgetToFocusOnActivate(ClassViewer);
+	}
+}
+
+void SSCreateBlueprintPicker::OnClassPicked(UClass* InChosenClass)
+{
+	ChosenClass = InChosenClass;
+}
+
+FReply SSCreateBlueprintPicker::OnClassPickerConfirmed()
+{
+	if (ChosenClass == nullptr)
+	{
+		FMessageDialog::Open(EAppMsgType::Ok, NSLOCTEXT("EditorFactories", "MustChooseClassWarning", "You must choose a class."));
+	}
+	else if (bIsReportingError)
+	{
+		FMessageDialog::Open(EAppMsgType::Ok, LOCTEXT("InvalidAssetname", "You must specify a valid asset name."));
+	}
+	else
+	{
+		bPressedOk = true;
+
+		if (WeakParentWindow.IsValid())
+		{
+			WeakParentWindow.Pin()->RequestDestroyWindow();
+		}
+	}
+	return FReply::Handled();
+}
+
+FReply SSCreateBlueprintPicker::OnClassPickerCanceled()
+{
+	if (WeakParentWindow.IsValid())
+	{
+		WeakParentWindow.Pin()->RequestDestroyWindow();
+	}
+	return FReply::Handled();
+}
+
+class SSCreateBlueprintPathPicker : public SCompoundWidget
+{
+	SLATE_BEGIN_ARGS(SSCreateBlueprintPathPicker)
+	{}
+
+	SLATE_ARGUMENT(TSharedPtr<SWindow>, ParentWindow)
+		SLATE_ARGUMENT(FString, AssetPath)
+		SLATE_END_ARGS()
+
+	void Construct(const FArguments& InArgs);
+
+	/** Callback when the selected asset path has changed. */
+	void OnSelectAssetPath(const FString& Path) { AssetPath = Path; }
+
+	/** Callback when the "ok" button is clicked. */
+	FReply OnClickOk();
+
+	/** Destroys the window when the operation is cancelled. */
+	FReply OnClickCancel();
+
+	/** A pointer to the window that is asking the user to select a parent class */
+	TWeakPtr<SWindow> WeakParentWindow;
+
+	FString AssetPath;
+
+	bool bPressedOk;
+};
+
+void SSCreateBlueprintPathPicker::Construct(const FArguments& InArgs)
+{
+	WeakParentWindow = InArgs._ParentWindow;
+
+	bPressedOk = false;
+	AssetPath = InArgs._AssetPath;
+
+	FContentBrowserModule& ContentBrowserModule = FModuleManager::LoadModuleChecked<FContentBrowserModule>("ContentBrowser");
+
+	FPathPickerConfig PathPickerConfig;
+	PathPickerConfig.DefaultPath = AssetPath;
+	PathPickerConfig.OnPathSelected = FOnPathSelected::CreateRaw(this, &SSCreateBlueprintPathPicker::OnSelectAssetPath);
+
+	ChildSlot
+	[
+		SNew(SVerticalBox)
+		+ SVerticalBox::Slot()
+		.FillHeight(1.0f)
+		[
+			ContentBrowserModule.Get().CreatePathPicker(PathPickerConfig)
+		]
+		+ SVerticalBox::Slot()
+		.HAlign(HAlign_Right)
+		.Padding(0, 20, 0, 0)
+		.AutoHeight()
+		[
+			SNew(SHorizontalBox)
+			+ SHorizontalBox::Slot()
+			.Padding(0, 2, 6, 0)
+			.AutoWidth()
+			[
+				SNew(SButton)
+				.VAlign(VAlign_Bottom)
+				.ContentPadding(FMargin(8, 2, 8, 2))
+				.OnClicked(this, &SSCreateBlueprintPathPicker::OnClickOk)
+				.ButtonStyle(FEditorStyle::Get(), "FlatButton.Success")
+				.TextStyle(FEditorStyle::Get(), "FlatButton.DefaultTextStyle")
+				.Text(LOCTEXT("OkButtonText", "OK"))
+			]
+			+ SHorizontalBox::Slot()
+			.Padding(0, 2, 0, 0)
+			.AutoWidth()
+			[
+				SNew(SButton)
+				.VAlign(VAlign_Bottom)
+				.ContentPadding(FMargin(8, 2, 8, 2))
+				.OnClicked(this, &SSCreateBlueprintPathPicker::OnClickCancel)
+				.ButtonStyle(FEditorStyle::Get(), "FlatButton.Default")
+				.TextStyle(FEditorStyle::Get(), "FlatButton.DefaultTextStyle")
+				.Text(LOCTEXT("CancelButtonText", "Cancel"))
+			]
+		]
+	];
+};
+
+FReply SSCreateBlueprintPathPicker::OnClickOk()
+{
+	bPressedOk = true;
+
+	if (WeakParentWindow.IsValid())
+	{
+		WeakParentWindow.Pin()->RequestDestroyWindow();
+	}
+	return FReply::Handled();
+}
+
+FReply SSCreateBlueprintPathPicker::OnClickCancel()
+{
+	if (WeakParentWindow.IsValid())
+	{
+		WeakParentWindow.Pin()->RequestDestroyWindow();
+	}
+	return FReply::Handled();
+}
+
+FReply SSCreateBlueprintPicker::OnPathPickerSummoned()
+{
+	// Create the window to pick the class
+	TSharedRef<SWindow> PickerWindow = SNew(SWindow)
+		.Title(LOCTEXT("CreateBlueprintFromActors_PickPath", "Select Path"))
+		.SizingRule(ESizingRule::UserSized)
+		.ClientSize(FVector2D(300.f, 400.f))
+		.SupportsMaximize(false)
+		.SupportsMinimize(false);
+
+	TSharedRef<SSCreateBlueprintPathPicker> PathPickerDialog = SNew(SSCreateBlueprintPathPicker)
+		.ParentWindow(PickerWindow)
+		.AssetPath(AssetPath);
+
+	PickerWindow->SetContent(PathPickerDialog);
+
+	GEditor->EditorAddModalWindow(PickerWindow);
+
+	if (PathPickerDialog->bPressedOk)
+	{
+		AssetPath = PathPickerDialog->AssetPath;
+		UpdateFilenameStatus();
+	}
+
+	return FReply::Handled();
+}
+
+void SSCreateBlueprintPicker::OnCustomAreaExpansionChanged(bool bExpanded)
+{
+	if (bExpanded && WeakParentWindow.IsValid())
+	{
+		WeakParentWindow.Pin().Get()->SetWidgetToFocusOnActivate(ClassViewer);
+	}
+}
+
+void SSCreateBlueprintPicker::OnFilenameChanged(const FText& InNewName)
+{
+	AssetName = InNewName.ToString();
+	UpdateFilenameStatus();
+}
+
+void SSCreateBlueprintPicker::UpdateFilenameStatus()
+{
+	FText ErrorText;
+	if (!FFileHelper::IsFilenameValidForSaving(AssetName, ErrorText) || !FName(*AssetName).IsValidObjectName(ErrorText))
+	{
+		FileNameWidget->SetError(ErrorText);
+		bIsReportingError = true;
+		return;
+	}
+	else
+	{
+		TArray<FAssetData> AssetData;
+		FAssetRegistryModule& AssetRegistryModule = FModuleManager::LoadModuleChecked<FAssetRegistryModule>("AssetRegistry");
+		AssetRegistryModule.Get().GetAssetsByPath(FName(*AssetPath), AssetData);
+
+		// Check to see if the name conflicts
+		for (const FAssetData& Data : AssetData)
+		{
+			const FString& AssetNameStr = Data.AssetName.ToString();
+			if (Data.AssetName.ToString() == AssetName)
+			{
+				FileNameWidget->SetError(LOCTEXT("AssetInUseError", "Asset name already in use!"));
+				bIsReportingError = true;
+				return;
+			}
+		}
+	}
+
+	FileNameWidget->SetError(FText::GetEmpty());
+	bIsReportingError = false;
+}
+
+EVisibility SSCreateBlueprintPicker::GetSelectButtonVisibility() const
+{
+	EVisibility ButtonVisibility = EVisibility::Hidden;
+	if (ChosenClass != nullptr && !bIsReportingError)
+	{
+		ButtonVisibility = EVisibility::Visible;
+	}
+	return ButtonVisibility;
+}
+
+/** Overridden from SWidget: Called when a key is pressed down - capturing copy */
+FReply SSCreateBlueprintPicker::OnKeyDown(const FGeometry& MyGeometry, const FKeyEvent& InKeyEvent)
+{
+	WeakParentWindow.Pin().Get()->SetWidgetToFocusOnActivate(ClassViewer);
+
+	if (InKeyEvent.GetKey() == EKeys::Escape)
+	{
+		return OnClassPickerCanceled();
+	}
+	else if (InKeyEvent.GetKey() == EKeys::Enter)
+	{
+		OnClassPickerConfirmed();
+	}
+	else
+	{
+		return ClassViewer->OnKeyDown(MyGeometry, InKeyEvent);
+	}
+	return FReply::Handled();
+}
+
 void FCreateBlueprintFromActorDialog::OpenDialog(ECreateBlueprintFromActorMode CreateMode, AActor* InActorOverride )
 {
 	ActorOverride = InActorOverride;
 
-	TSharedPtr<SWindow> PickBlueprintPathWidget;
-	SAssignNew(PickBlueprintPathWidget, SWindow)
-		.Title(LOCTEXT("SelectPath", "Select Path"))
-		.ToolTipText(LOCTEXT("SelectPathTooltip", "Select the path where the Blueprint will be created at"))
-		.ClientSize(FVector2D(400, 400));
-
-	TSharedPtr<SCreateAssetFromObject> CreateBlueprintFromActorDialog;
-	PickBlueprintPathWidget->SetContent
-	(
-		SAssignNew(CreateBlueprintFromActorDialog, SCreateAssetFromObject, PickBlueprintPathWidget)
-		.AssetFilenameSuffix(TEXT("Blueprint"))
-		.HeadingText(LOCTEXT("CreateBlueprintFromActor_Heading", "Blueprint Name"))
-		.CreateButtonText(LOCTEXT("CreateBlueprintFromActor_ButtonLabel", "Create Blueprint"))
-		.OnCreateAssetAction(FOnPathChosen::CreateStatic(FCreateBlueprintFromActorDialog::OnCreateBlueprint, CreateMode))
-	);
+	FClassViewerInitializationOptions ClassViewerOptions;
+	ClassViewerOptions.Mode = EClassViewerMode::ClassPicker;
+	ClassViewerOptions.DisplayMode = EClassViewerDisplayMode::TreeView;
+	ClassViewerOptions.bShowObjectRootClass = true;
+	ClassViewerOptions.bIsPlaceableOnly = true;
+	ClassViewerOptions.bIsBlueprintBaseOnly = true;
+	ClassViewerOptions.bShowUnloadedBlueprints = true;
+	ClassViewerOptions.bEnableClassDynamicLoading = true;
+	ClassViewerOptions.NameTypeToDisplay = EClassViewerNameTypeToDisplay::Dynamic;
 
-	TSharedPtr<SWindow> RootWindow = FGlobalTabmanager::Get()->GetRootWindow();
-	if (RootWindow.IsValid())
+	if (InActorOverride)
 	{
-		FSlateApplication::Get().AddWindowAsNativeChild(PickBlueprintPathWidget.ToSharedRef(), RootWindow.ToSharedRef());
+		class FBlueprintFromActorParentFilter : public IClassViewerFilter
+		{
+		public:
+			TSet<const UClass*> AllowedClass;
+
+			virtual bool IsClassAllowed(const FClassViewerInitializationOptions& InInitOptions, const UClass* InClass, TSharedRef< FClassViewerFilterFuncs > InFilterFuncs) override
+			{
+				return InFilterFuncs->IfInChildOfClassesSet(AllowedClass, InClass) == EFilterReturn::Passed;
+			}
+
+			virtual bool IsUnloadedClassAllowed(const FClassViewerInitializationOptions& InInitOptions, const TSharedRef< const IUnloadedBlueprintData > InUnloadedClassData, TSharedRef< FClassViewerFilterFuncs > InFilterFuncs) override
+			{
+				return InFilterFuncs->IfInChildOfClassesSet(AllowedClass, InUnloadedClassData) == EFilterReturn::Passed;
+			}
+		};
+
+		TSharedPtr<FBlueprintFromActorParentFilter> Filter = MakeShareable(new FBlueprintFromActorParentFilter);
+		ClassViewerOptions.ClassFilter = Filter;
+		Filter->AllowedClass.Add(InActorOverride->GetClass());
+
+		ClassViewerOptions.InitiallySelectedClass = InActorOverride->GetClass();
 	}
 	else
 	{
-		FSlateApplication::Get().AddWindow(PickBlueprintPathWidget.ToSharedRef());
+		ClassViewerOptions.InitiallySelectedClass = AActor::StaticClass();
+	}
+
+	// Create the window to pick the class
+	TSharedRef<SWindow> PickerWindow = SNew(SWindow)
+		.Title(LOCTEXT("CreateBlueprintFromActors","Create Blueprint From Selection"))
+		.SizingRule(ESizingRule::UserSized)
+		.ClientSize(FVector2D(600.f, 600.f))
+		.SupportsMaximize(false)
+		.SupportsMinimize(false);
+
+	TSharedRef<SSCreateBlueprintPicker> ClassPickerDialog = SNew(SSCreateBlueprintPicker)
+		.ParentWindow(PickerWindow)
+		.Options(ClassViewerOptions)
+		.CreateMode(CreateMode);
+
+	PickerWindow->SetContent(ClassPickerDialog);
+
+//	FSlateApplication::Get().AddWindow(PickerWindow);
+	GEditor->EditorAddModalWindow(PickerWindow);
+
+	if (ClassPickerDialog->bPressedOk)
+	{
+		FString NewAssetName = ClassPickerDialog->AssetPath / ClassPickerDialog->AssetName;
+
+		OnCreateBlueprint(NewAssetName, ClassPickerDialog->ChosenClass, ClassPickerDialog->CreateMode);
 	}
 }
 
-void FCreateBlueprintFromActorDialog::OnCreateBlueprint(const FString& InAssetPath, ECreateBlueprintFromActorMode CreateMode)
+void FCreateBlueprintFromActorDialog::OnCreateBlueprint(const FString& InAssetPath, UClass* ParentClass, ECreateBlueprintFromActorMode CreateMode)
 {
-	UBlueprint* Blueprint = NULL;
+	UBlueprint* Blueprint = nullptr;
 
 	switch (CreateMode) 
 	{
@@ -71,7 +689,7 @@ void FCreateBlueprintFromActorDialog::OnCreateBlueprint(const FString& InAssetPa
 		}
 
 			const bool bReplaceActor = true;
-			Blueprint = FKismetEditorUtilities::HarvestBlueprintFromActors(InAssetPath, Actors, bReplaceActor);
+			Blueprint = FKismetEditorUtilities::HarvestBlueprintFromActors(InAssetPath, Actors, bReplaceActor, ParentClass);
 	}
 		break;
 
@@ -88,7 +706,7 @@ void FCreateBlueprintFromActorDialog::OnCreateBlueprint(const FString& InAssetPa
 		}
 
 		const bool bReplaceActor = true;
-		Blueprint = FKismetEditorUtilities::CreateBlueprintFromActor(InAssetPath, ActorToUse, bReplaceActor);
+		Blueprint = FKismetEditorUtilities::CreateBlueprintFromActor(InAssetPath, ActorToUse, bReplaceActor, false, ParentClass);
 	}
 		break;
 
@@ -107,7 +725,7 @@ void FCreateBlueprintFromActorDialog::OnCreateBlueprint(const FString& InAssetPa
 			}
 
 			const bool bReplaceActor = true;
-			Blueprint = FKismetEditorUtilities::CreateBlueprintFromActors(InAssetPath, Actors, bReplaceActor);
+			Blueprint = FKismetEditorUtilities::CreateBlueprintFromActors(InAssetPath, Actors, bReplaceActor, ParentClass);
 		}
 		break;
 	}
diff --git a/Engine/Source/Editor/KismetWidgets/Public/CreateBlueprintFromActorDialog.h b/Engine/Source/Editor/KismetWidgets/Public/CreateBlueprintFromActorDialog.h
--- a/Engine/Source/Editor/KismetWidgets/Public/CreateBlueprintFromActorDialog.h
+++ b/Engine/Source/Editor/KismetWidgets/Public/CreateBlueprintFromActorDialog.h
@@ -44,7 +44,7 @@ private:
 	 * @param InAssetPath		Path of the asset to create
 	 * @param bInHarvest		The mode to use when creating a blueprint from the selected actors
 	 */
-	static void OnCreateBlueprint(const FString& InAssetPath, ECreateBlueprintFromActorMode CreateMode);
+	static void OnCreateBlueprint(const FString& InAssetPath, UClass* ParentClass, ECreateBlueprintFromActorMode CreateMode);
 
 	static TWeakObjectPtr<AActor> ActorOverride;
 };
diff --git a/Engine/Source/Editor/LevelEditor/Private/LevelEditorActions.cpp b/Engine/Source/Editor/LevelEditor/Private/LevelEditorActions.cpp
--- a/Engine/Source/Editor/LevelEditor/Private/LevelEditorActions.cpp
+++ b/Engine/Source/Editor/LevelEditor/Private/LevelEditorActions.cpp
@@ -2326,83 +2326,15 @@ void FLevelEditorActionCallbacks::CreateBlankBlueprintClass()
 	}
 }
 
-bool FLevelEditorActionCallbacks::CanHarvestSelectedActorsIntoBlueprintClass()
-{
-	return GEditor->GetSelectedActorCount() > 0;
-}
-
-void FLevelEditorActionCallbacks::HarvestSelectedActorsIntoBlueprintClass()
-{
-	FCreateBlueprintFromActorDialog::OpenDialog(ECreateBlueprintFromActorMode::Harvest);
-}
-
-ECreateBlueprintFromActorMode GetConvertSelectedActorsMode()
-{
-	if (GEditor->GetSelectedActorCount() == 0)
-	{
-		return ECreateBlueprintFromActorMode::None;
-	}
-
-	// Subclassing mode when we have 1 actor selected of a blueprint type
-	if (GEditor->GetSelectedActorCount() == 1)
-	{
-		AActor* Actor = Cast<AActor>(*GEditor->GetSelectedActorIterator());
-		if (FKismetEditorUtilities::CanCreateBlueprintOfClass(Actor->GetClass()))
-		{
-			return ECreateBlueprintFromActorMode::Subclass;
-		}
-	}
-
-	// Child Actor mode when we have 1 or more selected objects that are all placeable
-	for (FSelectionIterator It(GEditor->GetSelectedActorIterator()); It; ++It)
-	{
-		AActor* Actor = Cast<AActor>(*It);
-		if (Actor->GetClass()->HasAnyClassFlags(CLASS_NotPlaceable))
-		{
-			return ECreateBlueprintFromActorMode::None;
-		}
-	}
-
-	return ECreateBlueprintFromActorMode::ChildActor;
-}
-
 bool FLevelEditorActionCallbacks::CanConvertSelectedActorsIntoBlueprintClass()
 {
-	const bool bCanConvert = GetConvertSelectedActorsMode() != ECreateBlueprintFromActorMode::None;
-	return bCanConvert;
+	return GEditor->GetSelectedActorCount() > 0;
 }
 
 void FLevelEditorActionCallbacks::ConvertSelectedActorsIntoBlueprintClass()
 {
-	ECreateBlueprintFromActorMode CreateMode = GetConvertSelectedActorsMode();
-	if (ensure(CreateMode != ECreateBlueprintFromActorMode::None))
-	{
-		FCreateBlueprintFromActorDialog::OpenDialog(CreateMode);
-	}
-}
-
-FText FLevelEditorActionCallbacks::GetConvertSelectedActorsIntoBlueprintClassLabel()
-{
-	if (GetConvertSelectedActorsMode() == ECreateBlueprintFromActorMode::ChildActor)
-{
-		return LOCTEXT("ConvertSelectionToBlueprintViaChildActor_Label", "Convert Selected Actors to Blueprint Class...");
-	}
-	else
-	{
-		return LOCTEXT("ConvertSelectionToBlueprintViaSubclass_Label", "Convert Selected Actor to Blueprint Class...");
-	}
-}
-
-FText FLevelEditorActionCallbacks::GetConvertSelectedActorsIntoBlueprintClassTooltip()
-{
-	if (GetConvertSelectedActorsMode() == ECreateBlueprintFromActorMode::ChildActor)
-	{
-		return LOCTEXT("ConvertSelectionToBlueprintViaChildActor_Tooltip", "Replace the selected actors with a new Blueprint subclass based on Actor with each of the selected Actors as a Child Actor");
-	}
-	else
-	{
-		return LOCTEXT("ConvertSelectionToBlueprintViaSubclass_Tooltip", "Replace the selected actor with a new Blueprint subclass based on the class of the selected Actor");
-	}
+	const ECreateBlueprintFromActorMode DefaultCreateMode = (GEditor->GetSelectedActorCount() == 1 ? ECreateBlueprintFromActorMode::Subclass : ECreateBlueprintFromActorMode::ChildActor);
+	FCreateBlueprintFromActorDialog::OpenDialog(DefaultCreateMode);
 }
 
 void FLevelEditorActionCallbacks::CheckOutProjectSettingsConfig( )
@@ -3358,8 +3290,7 @@ void FLevelEditorCommands::RegisterCommands()
 	UI_COMMAND( OpenLevelBlueprint, "Open Level Blueprint", "Edit the Level Blueprint for the current level", EUserInterfaceActionType::Button, FInputChord() );
 	UI_COMMAND( CheckOutProjectSettingsConfig, "Check Out", "Checks out the project settings config file so the game mode can be set.", EUserInterfaceActionType::Button, FInputChord() );
 	UI_COMMAND( CreateBlankBlueprintClass, "New Empty Blueprint Class...", "Create a new Blueprint Class", EUserInterfaceActionType::Button, FInputChord() );
-	UI_COMMAND( ConvertSelectionToBlueprintViaHarvest, "Convert Selected Components to Blueprint Class...", "Replace all of the selected actors with a new Blueprint Class based on Actor that contains the components", EUserInterfaceActionType::Button, FInputChord() );
-	UI_COMMAND( ConvertSelectionToBlueprintViaSubclass, "Convert Selected Actor to Blueprint Class...", "Replace the selected actor with a new Blueprint subclass based on the class of the selected Actor", EUserInterfaceActionType::Button, FInputChord() );
+	UI_COMMAND( ConvertSelectionToBlueprint, "Convert Selection to Blueprint Class...", "Replace all of the selected actors with a new Blueprint Class", EUserInterfaceActionType::Button, FInputChord() );
 
 	UI_COMMAND( ShowTransformWidget, "Show Transform Widget", "Toggles the visibility of the transform widgets", EUserInterfaceActionType::ToggleButton, FInputChord() );
 	UI_COMMAND( AllowTranslucentSelection, "Allow Translucent Selection", "Allows translucent objects to be selected", EUserInterfaceActionType::ToggleButton, FInputChord(EKeys::T) );
diff --git a/Engine/Source/Editor/LevelEditor/Private/LevelEditorToolBar.cpp b/Engine/Source/Editor/LevelEditor/Private/LevelEditorToolBar.cpp
--- a/Engine/Source/Editor/LevelEditor/Private/LevelEditorToolBar.cpp
+++ b/Engine/Source/Editor/LevelEditor/Private/LevelEditorToolBar.cpp
@@ -2315,11 +2315,7 @@ void FLevelEditorToolBar::RegisterOpenBlueprintMenu()
 		Section.AddMenuEntry(FLevelEditorCommands::Get().CreateBlankBlueprintClass);
 
 		// Convert selection to BP
-		Section.AddMenuEntry(FLevelEditorCommands::Get().ConvertSelectionToBlueprintViaHarvest);
-		Section.AddMenuEntry(
-			FLevelEditorCommands::Get().ConvertSelectionToBlueprintViaSubclass,
-			TAttribute<FText>::Create(FLevelEditorActionCallbacks::GetConvertSelectedActorsIntoBlueprintClassLabel),
-			TAttribute<FText>::Create(FLevelEditorActionCallbacks::GetConvertSelectedActorsIntoBlueprintClassTooltip));
+		Section.AddMenuEntry(FLevelEditorCommands::Get().ConvertSelectionToBlueprint);
 
 		// Open an existing Blueprint Class...
 		FSlateIcon OpenBPIcon(FEditorStyle::GetStyleSetName(), "LevelEditor.OpenClassBlueprint");
diff --git a/Engine/Source/Editor/LevelEditor/Private/SLevelEditor.cpp b/Engine/Source/Editor/LevelEditor/Private/SLevelEditor.cpp
--- a/Engine/Source/Editor/LevelEditor/Private/SLevelEditor.cpp
+++ b/Engine/Source/Editor/LevelEditor/Private/SLevelEditor.cpp
@@ -106,12 +106,7 @@ void SLevelEditor::BindCommands()
 		FExecuteAction::CreateStatic( &FLevelEditorActionCallbacks::CreateBlankBlueprintClass ) );
 
 	LevelEditorCommands->MapAction(
-		Actions.ConvertSelectionToBlueprintViaHarvest,
-		FExecuteAction::CreateStatic( &FLevelEditorActionCallbacks::HarvestSelectedActorsIntoBlueprintClass ),
-		FCanExecuteAction::CreateStatic( &FLevelEditorActionCallbacks::CanHarvestSelectedActorsIntoBlueprintClass ) );
-
-	LevelEditorCommands->MapAction(
-		Actions.ConvertSelectionToBlueprintViaSubclass,
+		Actions.ConvertSelectionToBlueprint,
 		FExecuteAction::CreateStatic( &FLevelEditorActionCallbacks::ConvertSelectedActorsIntoBlueprintClass ),
 		FCanExecuteAction::CreateStatic( &FLevelEditorActionCallbacks::CanConvertSelectedActorsIntoBlueprintClass ) );
 
diff --git a/Engine/Source/Editor/LevelEditor/Public/LevelEditorActions.h b/Engine/Source/Editor/LevelEditor/Public/LevelEditorActions.h
--- a/Engine/Source/Editor/LevelEditor/Public/LevelEditorActions.h
+++ b/Engine/Source/Editor/LevelEditor/Public/LevelEditorActions.h
@@ -555,8 +555,7 @@ public:
 	TSharedPtr< FUICommandInfo > OpenLevelBlueprint;
 	TSharedPtr< FUICommandInfo > CheckOutProjectSettingsConfig;
 	TSharedPtr< FUICommandInfo > CreateBlankBlueprintClass;
-	TSharedPtr< FUICommandInfo > ConvertSelectionToBlueprintViaHarvest;
-	TSharedPtr< FUICommandInfo > ConvertSelectionToBlueprintViaSubclass;
+	TSharedPtr< FUICommandInfo > ConvertSelectionToBlueprint;
 
 	/** Editor mode commands */
 	TArray< TSharedPtr< FUICommandInfo > > EditorModeCommands;
@@ -1178,24 +1177,12 @@ public:
 	/** Helps the user create a Blueprint class */
 	static void CreateBlankBlueprintClass();
 
-	/** Can call HarvestSelectedActorsIntoBlueprintClass right now?  (is anything selected) */
-	static bool CanHarvestSelectedActorsIntoBlueprintClass();
-
-	/** Harvest all of the components in the selected actors into a Blueprint, and replace the instances with one new actor */
-	static void HarvestSelectedActorsIntoBlueprintClass();
-
-	/** Can call SubclassSelectedActorIntoBlueprintClass right now? */
+	/** Can the selected actors be converted to a blueprint class in any of the supported ways? */
 	static bool CanConvertSelectedActorsIntoBlueprintClass();
 
-	/** Convert the selected actor into a Blueprint (via subclassing or child actors) */
+	/** Bring up the convert actors to blueprint UI */
 	static void ConvertSelectedActorsIntoBlueprintClass();
 
-	/** Get the label to display for the convert selected actors to blueprint menu option */
-	static FText GetConvertSelectedActorsIntoBlueprintClassLabel();
-
-	/** Get the tooltip to display for the convert selected actors to blueprint menu option */
-	static FText GetConvertSelectedActorsIntoBlueprintClassTooltip();
-
 	/** Shows only selected actors, hiding any unselected actors and unhiding any selected hidden actors. */
 	static void OnShowOnlySelectedActors();
 
diff --git a/Engine/Source/Editor/UnrealEd/Private/Kismet2/Kismet2.cpp b/Engine/Source/Editor/UnrealEd/Private/Kismet2/Kismet2.cpp
--- a/Engine/Source/Editor/UnrealEd/Private/Kismet2/Kismet2.cpp
+++ b/Engine/Source/Editor/UnrealEd/Private/Kismet2/Kismet2.cpp
@@ -1036,27 +1036,27 @@ UPackage* CreateBlueprintPackage(const FString& Path, FString& OutAssetName)
 	return CreatePackage(nullptr, *PackageName);
 }
 
-UBlueprint* FKismetEditorUtilities::CreateBlueprintFromActor(const FString& Path, AActor* Actor, const bool bReplaceActor, bool bKeepMobility /*= false*/)
+UBlueprint* FKismetEditorUtilities::CreateBlueprintFromActor(const FString& Path, AActor* Actor, const bool bReplaceActor, bool bKeepMobility /*= false*/, UClass* ParentClassOverride)
 {
 	UBlueprint* NewBlueprint = nullptr;
 	FString AssetName;
 
 	if (UPackage* Package = CreateBlueprintPackage(Path, AssetName))
 	{
-		NewBlueprint = CreateBlueprintFromActor(FName(*AssetName), Package, Actor, bReplaceActor, bKeepMobility);
+		NewBlueprint = CreateBlueprintFromActor(FName(*AssetName), Package, Actor, bReplaceActor, bKeepMobility, ParentClassOverride);
 	}
 
 	return NewBlueprint;
 }
 
-UBlueprint* FKismetEditorUtilities::CreateBlueprintFromActors(const FString& Path, const TArray<AActor*>& Actors, const bool bReplaceInWorld)
+UBlueprint* FKismetEditorUtilities::CreateBlueprintFromActors(const FString& Path, const TArray<AActor*>& Actors, const bool bReplaceInWorld, UClass* ParentClass)
 {
 	UBlueprint* NewBlueprint = nullptr;
 	FString AssetName;
 
 	if (UPackage* Package = CreateBlueprintPackage(Path, AssetName))
 	{
-		NewBlueprint = CreateBlueprintFromActors(FName(*AssetName), Package, Actors, bReplaceInWorld);
+		NewBlueprint = CreateBlueprintFromActors(FName(*AssetName), Package, Actors, bReplaceInWorld, ParentClass);
 	}
 
 	return NewBlueprint;
@@ -1341,7 +1341,7 @@ private:
 	friend class FKismetEditorUtilities;
 };
 
-UBlueprint* FKismetEditorUtilities::CreateBlueprintFromActor(const FName BlueprintName, UObject* Outer, AActor* Actor, const bool bReplaceActor, bool bKeepMobility /*= false*/)
+UBlueprint* FKismetEditorUtilities::CreateBlueprintFromActor(const FName BlueprintName, UObject* Outer, AActor* Actor, const bool bReplaceActor, bool bKeepMobility /*= false*/, UClass* ParentClassOverride)
 {
 	UBlueprint* NewBlueprint = nullptr;
 
@@ -1349,8 +1349,21 @@ UBlueprint* FKismetEditorUtilities::CreateBlueprintFromActor(const FName Bluepri
 	{
 		if (Outer != nullptr)
 		{
+			if (ParentClassOverride)
+			{
+				if (!ParentClassOverride->IsChildOf(Actor->GetClass()))
+				{
+					// Invalid input, use ActorClass instead?
+					return nullptr;
+				}
+			}
+			else
+			{
+				ParentClassOverride = Actor->GetClass();
+			}
+
 			// We don't have a factory, but we can still try to create a blueprint for this actor class
-			NewBlueprint = FKismetEditorUtilities::CreateBlueprint( Actor->GetClass(), Outer, BlueprintName, EBlueprintType::BPTYPE_Normal, UBlueprint::StaticClass(), UBlueprintGeneratedClass::StaticClass(), FName("CreateFromActor") );
+			NewBlueprint = FKismetEditorUtilities::CreateBlueprint(ParentClassOverride, Outer, BlueprintName, EBlueprintType::BPTYPE_Normal, UBlueprint::StaticClass(), UBlueprintGeneratedClass::StaticClass(), FName("CreateFromActor") );
 		}
 
 		if (NewBlueprint != nullptr)
@@ -1570,7 +1583,7 @@ void CreateBlueprintFromActors_Internal(UBlueprint* Blueprint, const TArray<AAct
 	GEditor->GetEditorSubsystem<UAssetEditorSubsystem>()->OpenEditorForAsset(Blueprint);
 }
 
-UBlueprint* FKismetEditorUtilities::CreateBlueprintFromActors(const FName BlueprintName, UPackage* Package, const TArray<AActor*>& Actors, const bool bReplaceInWorld)
+UBlueprint* FKismetEditorUtilities::CreateBlueprintFromActors(const FName BlueprintName, UPackage* Package, const TArray<AActor*>& Actors, const bool bReplaceInWorld, UClass* ParentClass)
 {
 	auto AssemblyFunction = [](const FBlueprintAssemblyProps& AssemblyProps)
 	{
@@ -1619,7 +1632,7 @@ UBlueprint* FKismetEditorUtilities::CreateBlueprintFromActors(const FName Bluepr
 		if (Package != nullptr)
 		{
 			// We don't have a factory, but we can still try to create a blueprint for this actor class
-			Blueprint = FKismetEditorUtilities::CreateBlueprint(AActor::StaticClass(), Package, BlueprintName, EBlueprintType::BPTYPE_Normal, UBlueprint::StaticClass(), UBlueprintGeneratedClass::StaticClass(), FName("CreateFromActors"));
+			Blueprint = FKismetEditorUtilities::CreateBlueprint(ParentClass, Package, BlueprintName, EBlueprintType::BPTYPE_Normal, UBlueprint::StaticClass(), UBlueprintGeneratedClass::StaticClass(), FName("CreateFromActors"));
 		}
 
 		if (Blueprint != nullptr)
@@ -1631,20 +1644,20 @@ UBlueprint* FKismetEditorUtilities::CreateBlueprintFromActors(const FName Bluepr
 	return Blueprint;
 }
 
-UBlueprint* FKismetEditorUtilities::HarvestBlueprintFromActors(const FString& Path, const TArray<AActor*>& Actors, bool bReplaceInWorld)
+UBlueprint* FKismetEditorUtilities::HarvestBlueprintFromActors(const FString& Path, const TArray<AActor*>& Actors, bool bReplaceInWorld, UClass* ParentClass)
 {
 	UBlueprint* NewBlueprint = nullptr;
 	FString AssetName;
 
 	if (UPackage* Package = CreateBlueprintPackage(Path, AssetName))
 	{
-		NewBlueprint = HarvestBlueprintFromActors(FName(*AssetName), Package, Actors, bReplaceInWorld);
+		NewBlueprint = HarvestBlueprintFromActors(FName(*AssetName), Package, Actors, bReplaceInWorld, ParentClass);
 	}
 
 	return NewBlueprint;
 }
 
-UBlueprint* FKismetEditorUtilities::HarvestBlueprintFromActors(const FName BlueprintName, UPackage* Package, const TArray<AActor*>& Actors, bool bReplaceInWorld)
+UBlueprint* FKismetEditorUtilities::HarvestBlueprintFromActors(const FName BlueprintName, UPackage* Package, const TArray<AActor*>& Actors, bool bReplaceInWorld, UClass* ParentClass)
 {
 	auto AssemblyFunction = [](const FBlueprintAssemblyProps& AssemblyProps)
 	{
@@ -1690,7 +1703,7 @@ UBlueprint* FKismetEditorUtilities::HarvestBlueprintFromActors(const FName Bluep
 		if (Package != nullptr)
 		{
 			// We don't have a factory, but we can still try to create a blueprint for this actor class
-			Blueprint = FKismetEditorUtilities::CreateBlueprint(AActor::StaticClass(), Package, BlueprintName, EBlueprintType::BPTYPE_Normal, UBlueprint::StaticClass(), UBlueprintGeneratedClass::StaticClass(), FName("HarvestFromActors"));
+			Blueprint = FKismetEditorUtilities::CreateBlueprint(ParentClass, Package, BlueprintName, EBlueprintType::BPTYPE_Normal, UBlueprint::StaticClass(), UBlueprintGeneratedClass::StaticClass(), FName("HarvestFromActors"));
 		}
 
 		if (Blueprint != nullptr)
diff --git a/Engine/Source/Editor/UnrealEd/Public/Kismet2/KismetEditorUtilities.h b/Engine/Source/Editor/UnrealEd/Public/Kismet2/KismetEditorUtilities.h
--- a/Engine/Source/Editor/UnrealEd/Public/Kismet2/KismetEditorUtilities.h
+++ b/Engine/Source/Editor/UnrealEd/Public/Kismet2/KismetEditorUtilities.h
@@ -157,7 +157,7 @@ public:
 	 * @param bKeepMobility			If true, The mobility of each actor components will be copy
 	 * @return The blueprint created from the actor
 	 */
-	static UBlueprint* CreateBlueprintFromActor(const FString& Path, AActor* Actor, bool bReplaceActor, bool bKeepMobility = false );
+	static UBlueprint* CreateBlueprintFromActor(const FString& Path, AActor* Actor, bool bReplaceActor, bool bKeepMobility = false, UClass* ParentClassOverride = nullptr);
 
 	/** 
 	 * Take an Actor and generate a blueprint based on it. Uses the Actors type as the parent class. 
@@ -168,7 +168,7 @@ public:
 	 * @param bKeepMobility			If true, The mobility of each actor components will be copy
 	 * @return The blueprint created from the actor
 	 */
-	static UBlueprint* CreateBlueprintFromActor(const FName BlueprintName, UObject* Outer, AActor* Actor, bool bReplaceInWorld, bool bKeepMobility = false);
+	static UBlueprint* CreateBlueprintFromActor(const FName BlueprintName, UObject* Outer, AActor* Actor, bool bReplaceInWorld, bool bKeepMobility = false, UClass* ParentClassOverride = nullptr);
 
 	/** 
 	 * Take a list of Actors and generate a blueprint based on it using the Actors as templates for child actor components.
@@ -177,7 +177,7 @@ public:
 	 * @param bReplaceActor			If true, replace the actor in the scene with one based on the created blueprint
 	 * @return The blueprint created from the actor
 	 */
-	static UBlueprint* CreateBlueprintFromActors(const FString& Path, const TArray<AActor*>& Actors, bool bReplaceActor);
+	static UBlueprint* CreateBlueprintFromActors(const FString& Path, const TArray<AActor*>& Actors, bool bReplaceActor, UClass* ParentClass = AActor::StaticClass());
 
 	/** 
 	 * Take a list of Actors and generate a blueprint based on it using the Actors as templates for child actor components.
@@ -188,7 +188,7 @@ public:
 	 * @param bKeepMobility			If true, The mobility of each actor components will be copy
 	 * @return The blueprint created from the actor
 	 */
-	static UBlueprint* CreateBlueprintFromActors(const FName BlueprintName, UPackage* Package, const TArray<AActor*>& Actors, bool bReplaceInWorld);
+	static UBlueprint* CreateBlueprintFromActors(const FName BlueprintName, UPackage* Package, const TArray<AActor*>& Actors, bool bReplaceInWorld, UClass* ParentClass = AActor::StaticClass());
 	
 	/** 
 	 * Take a list of Actors and generate a blueprint  by harvesting the components they have. Uses AActor as parent class type as the parent class. 
@@ -197,7 +197,7 @@ public:
 	 * @param bReplaceInWorld		If true, replace the selected actors in the scene with one based on the created blueprint
 	 * @return The blueprint created from the actors
 	 */
-	static UBlueprint* HarvestBlueprintFromActors(const FString& Path, const TArray<AActor*>& Actors, bool bReplaceInWorld);
+	static UBlueprint* HarvestBlueprintFromActors(const FString& Path, const TArray<AActor*>& Actors, bool bReplaceInWorld, UClass* ParentClass = AActor::StaticClass());
 
 	/**
 	 * Take a list of Actors and generate a blueprint  by harvesting the components they have. Uses AActor as parent class type as the parent class.
@@ -207,7 +207,7 @@ public:
 	 * @param bReplaceInWorld		If true, replace the selected actors in the scene with one based on the created blueprint
 	 * @return The blueprint created from the actors
 	 */
-	static UBlueprint* HarvestBlueprintFromActors(const FName BlueprintName, UPackage* Package, const TArray<AActor*>& Actors, bool bReplaceInWorld);
+	static UBlueprint* HarvestBlueprintFromActors(const FName BlueprintName, UPackage* Package, const TArray<AActor*>& Actors, bool bReplaceInWorld, UClass* ParentClass = AActor::StaticClass());
 
 	/** 
 	 * Creates a new blueprint instance and replaces the provided actor list with the new actor
