You can no longer harvest components if there are no blueprint spawnable components.
Provide a tooltip explaining why a create mode is disabled
#rnx

diff --git a/Engine/Source/Editor/KismetWidgets/Private/CreateBlueprintFromActorDialog.cpp b/Engine/Source/Editor/KismetWidgets/Private/CreateBlueprintFromActorDialog.cpp
--- a/Engine/Source/Editor/KismetWidgets/Private/CreateBlueprintFromActorDialog.cpp
+++ b/Engine/Source/Editor/KismetWidgets/Private/CreateBlueprintFromActorDialog.cpp
@@ -19,6 +19,7 @@
 #include "Engine/Selection.h"
 #include "Editor.h"
 #include "Kismet2/KismetEditorUtilities.h"
+#include "EdGraphSchema_K2.h"
 #include "Framework/Notifications/NotificationManager.h"
 #include "Widgets/Notifications/SNotificationList.h"
 #include "AssetRegistryModule.h"
@@ -86,6 +87,8 @@ public:
 		} 	
 	};
 
+	FText GetCreateMethodTooltip(ECreateBlueprintFromActorMode InCreateMode, bool bEnabled) const;
+
 	FSlateColor GetCreateModeTextColor(ECreateBlueprintFromActorMode InCreateMode) const
 	{ 
 		return (CreateMode == InCreateMode ? FSlateColor(FLinearColor(0, 0, 0)) : FSlateColor(FLinearColor(0.72f, 0.72f, 0.72f, 1.f))); 
@@ -122,6 +125,102 @@ public:
 	bool bIsReportingError;
 };
 
+ECreateBlueprintFromActorMode FCreateBlueprintFromActorDialog::GetValidCreationMethods()
+{
+	int32 NumSelectedActors = 0;
+
+	bool bCanHarvestComponents = false;
+	bool bCanSubclass = true;
+	bool bCanCreatePrefab = true;
+
+	for (FSelectionIterator Iter(*GEditor->GetSelectedActors()); Iter; ++Iter)
+	{
+		AActor* Actor = Cast<AActor>(*Iter);
+		if (Actor)
+		{
+			if (NumSelectedActors == 0)
+			{
+				bCanSubclass = FKismetEditorUtilities::CanCreateBlueprintOfClass(Actor->GetClass());
+			}
+
+			if (bCanCreatePrefab && Actor->GetClass()->HasAnyClassFlags(CLASS_NotPlaceable))
+			{
+				bCanCreatePrefab = false;
+			}
+
+			if (!bCanHarvestComponents)
+			{
+				for (UActorComponent* Component : Actor->GetComponents())
+				{
+					if (Component && FKismetEditorUtilities::IsClassABlueprintSpawnableComponent(Component->GetClass()))
+					{
+						bCanHarvestComponents = true;
+						break;
+					}
+				}
+			}
+		}
+		++NumSelectedActors;
+	}
+
+	ECreateBlueprintFromActorMode ValidCreationMethods = ECreateBlueprintFromActorMode::None;
+	if (NumSelectedActors > 0)
+	{
+		if (bCanHarvestComponents)
+		{ 
+			ValidCreationMethods |= ECreateBlueprintFromActorMode::Harvest;
+		}
+		if (bCanCreatePrefab)
+		{
+			ValidCreationMethods |= ECreateBlueprintFromActorMode::ChildActor;
+		}
+		if (bCanSubclass && NumSelectedActors == 1)
+		{
+			ValidCreationMethods |= ECreateBlueprintFromActorMode::Subclass;
+		}
+	}
+
+	return ValidCreationMethods;
+}
+
+FText SSCreateBlueprintPicker::GetCreateMethodTooltip(ECreateBlueprintFromActorMode InCreateMode, bool bEnabled) const
+{
+	if (!bEnabled)
+	{
+		switch (InCreateMode)
+		{
+		case ECreateBlueprintFromActorMode::Subclass:
+		{
+			int32 NumSelectedActors = 0;
+			UClass* SelectedActorClass = nullptr;
+			for (FSelectionIterator Iter(*GEditor->GetSelectedActors()); Iter; ++Iter)
+			{
+				if (AActor* Actor = Cast<AActor>(*Iter))
+				{
+					SelectedActorClass = Actor->GetClass();
+				}
+				++NumSelectedActors;
+			}
+
+			if (NumSelectedActors == 1)
+			{
+				return FText::Format(LOCTEXT("SubClassDisabled_InvalidBlueprintType", "Cannot create blueprint subclass of '{0}'."), (SelectedActorClass ? SelectedActorClass->GetDisplayNameText() : FText::GetEmpty()));
+			}
+			else
+			{
+				return LOCTEXT("SubClassDisabled_MultipleSelection", "Cannot subclass when multiple actors are selected.");
+			}
+		}
+
+		case ECreateBlueprintFromActorMode::Harvest:
+			return LOCTEXT("HavestDisabled", "No harvestable components in selected actors.");
+
+		}
+	}
+
+	return FText::GetEmpty();
+}
+
 void SSCreateBlueprintPicker::Construct(const FArguments& InArgs)
 {
 	WeakParentWindow = InArgs._ParentWindow;
@@ -139,38 +238,23 @@ void SSCreateBlueprintPicker::Construct(const FArguments& InArgs)
 	FString PackageName;
 	AssetPath = ContentBrowserModule.Get().GetCurrentPath();
 
-	int32 NumSelectedActors = 0;
+	ECreateBlueprintFromActorMode ValidCreateMethods = FCreateBlueprintFromActorDialog::GetValidCreationMethods();
 
-	bool bCanHarvestComponents = false;
-	bool bCanSubclass = true;
-	bool bCanCreatePrefab = true;
+	const bool bCanHarvestComponents = !!(ValidCreateMethods & ECreateBlueprintFromActorMode::Harvest);
+	const bool bCanSubclass = !!(ValidCreateMethods & ECreateBlueprintFromActorMode::Subclass);
+	const bool bCanCreatePrefab = !!(ValidCreateMethods & ECreateBlueprintFromActorMode::ChildActor);
 
 	for (FSelectionIterator Iter(*GEditor->GetSelectedActors()); Iter; ++Iter)
 	{
 		AActor* Actor = Cast<AActor>(*Iter);
 		if (Actor)
 		{
-			if (NumSelectedActors == 0)
-			{
-				AssetName += Actor->GetActorLabel();
-				AssetName += TEXT("_");
-
-				bCanSubclass = FKismetEditorUtilities::CanCreateBlueprintOfClass(Actor->GetClass());
-			}
-
-			if (Actor->GetClass()->HasAnyClassFlags(CLASS_NotPlaceable))
-			{
-				bCanCreatePrefab = false;
-			}
-
-			bCanHarvestComponents = true;
-
-			++NumSelectedActors;
+			AssetName += Actor->GetActorLabel();
+			AssetName += TEXT("_");
+			break;
 		}
 	}
 
-	bCanSubclass = bCanSubclass && NumSelectedActors == 1;
-
 	AssetName = UPackageTools::SanitizePackageName(AssetName + TEXT("Blueprint"));
 
 	FString BasePath = AssetPath / AssetName;
@@ -209,6 +293,7 @@ void SSCreateBlueprintPicker::Construct(const FArguments& InArgs)
 			.IsEnabled(CreateModeDetails[Index].bEnabled)
 			.IsChecked_Raw(this, &SSCreateBlueprintPicker::IsCreateModeChecked, CreateModeDetails[Index].CreateMode)
 			.OnCheckStateChanged(this, &SSCreateBlueprintPicker::OnCreateModeChanged, CreateModeDetails[Index].CreateMode)
+			.ToolTipText_Raw(this, &SSCreateBlueprintPicker::GetCreateMethodTooltip, CreateModeDetails[Index].CreateMode, CreateModeDetails[Index].bEnabled)
 			[
 				SNew(SHorizontalBox)
 				+ SHorizontalBox::Slot()
@@ -675,39 +760,39 @@ void FCreateBlueprintFromActorDialog::OnCreateBlueprint(const FString& InAssetPa
 	switch (CreateMode) 
 	{
 		case ECreateBlueprintFromActorMode::Harvest:
-	{
-		TArray<AActor*> Actors;
-
-		USelection* SelectedActors = GEditor->GetSelectedActors();
-		for(FSelectionIterator Iter(*SelectedActors); Iter; ++Iter)
 		{
-			// We only care about actors that are referenced in the world for literals, and also in the same level as this blueprint
-				if (AActor* Actor = Cast<AActor>(*Iter))
+			TArray<AActor*> Actors;
+
+			USelection* SelectedActors = GEditor->GetSelectedActors();
+			for(FSelectionIterator Iter(*SelectedActors); Iter; ++Iter)
 			{
-				Actors.Add(Actor);
+				// We only care about actors that are referenced in the world for literals, and also in the same level as this blueprint
+					if (AActor* Actor = Cast<AActor>(*Iter))
+				{
+					Actors.Add(Actor);
+				}
 			}
-		}
 
 			const bool bReplaceActor = true;
 			Blueprint = FKismetEditorUtilities::HarvestBlueprintFromActors(InAssetPath, Actors, bReplaceActor, ParentClass);
-	}
+		}
 		break;
 
 		case ECreateBlueprintFromActorMode::Subclass:
-	{
-		AActor* ActorToUse = ActorOverride.Get();
-
-		if( !ActorToUse )
 		{
-			TArray< UObject* > SelectedActors;
-			GEditor->GetSelectedActors()->GetSelectedObjects(AActor::StaticClass(), SelectedActors);
+			AActor* ActorToUse = ActorOverride.Get();
+
+			if (!ActorToUse)
+			{
+				TArray< UObject* > SelectedActors;
+				GEditor->GetSelectedActors()->GetSelectedObjects(AActor::StaticClass(), SelectedActors);
 				check(SelectedActors.Num() == 1);
-			ActorToUse =  Cast<AActor>(SelectedActors[0]);
-		}
+				ActorToUse = Cast<AActor>(SelectedActors[0]);
+			}
 
-		const bool bReplaceActor = true;
-		Blueprint = FKismetEditorUtilities::CreateBlueprintFromActor(InAssetPath, ActorToUse, bReplaceActor, false, ParentClass);
-	}
+			const bool bReplaceActor = true;
+			Blueprint = FKismetEditorUtilities::CreateBlueprintFromActor(InAssetPath, ActorToUse, bReplaceActor, false, ParentClass);
+		}
 		break;
 
 		case ECreateBlueprintFromActorMode::ChildActor:
@@ -730,7 +815,7 @@ void FCreateBlueprintFromActorDialog::OnCreateBlueprint(const FString& InAssetPa
 		break;
 	}
 
-	if(Blueprint)
+	if (Blueprint)
 	{
 		// Select the newly created blueprint in the content browser, but don't activate the browser
 		TArray<UObject*> Objects;
diff --git a/Engine/Source/Editor/KismetWidgets/Public/CreateBlueprintFromActorDialog.h b/Engine/Source/Editor/KismetWidgets/Public/CreateBlueprintFromActorDialog.h
--- a/Engine/Source/Editor/KismetWidgets/Public/CreateBlueprintFromActorDialog.h
+++ b/Engine/Source/Editor/KismetWidgets/Public/CreateBlueprintFromActorDialog.h
@@ -11,11 +11,12 @@
 /** Enum defining which mode to use when creating the blueprint from the selected actors */
 enum class ECreateBlueprintFromActorMode : uint8
 {
-	None,        // Indicates that a user is unable to create blueprints from the currently selected actor set
-	Harvest,     // Harvest all the components of the selected actors and create an actor blueprint with those components in it.
-	Subclass,     // Create a subclass of the selected Actor's class with defaults from the selected Actor. Valid only when there is a single selected actor.
-	ChildActor    // Create an actor blueprint with each selected actor as a child actor component
+	None,             // Indicates that a user is unable to create blueprints from the currently selected actor set
+	Harvest = 1,      // Harvest all the components of the selected actors and create an actor blueprint with those components in it.
+	Subclass = 2,     // Create a subclass of the selected Actor's class with defaults from the selected Actor. Valid only when there is a single selected actor.
+	ChildActor = 4    // Create an actor blueprint with each selected actor as a child actor component
 };
+ENUM_CLASS_FLAGS(ECreateBlueprintFromActorMode)
 
 class FCreateBlueprintFromActorDialog
 {
@@ -36,6 +37,11 @@ public:
 	}
 
 
+	/** 
+	 * Static function that returns which create modes are valid given the current selection set
+	 */
+	static KISMETWIDGETS_API ECreateBlueprintFromActorMode GetValidCreationMethods();
+
 private:
 
 	/** 
diff --git a/Engine/Source/Editor/LevelEditor/Private/LevelEditorActions.cpp b/Engine/Source/Editor/LevelEditor/Private/LevelEditorActions.cpp
--- a/Engine/Source/Editor/LevelEditor/Private/LevelEditorActions.cpp
+++ b/Engine/Source/Editor/LevelEditor/Private/LevelEditorActions.cpp
@@ -2328,12 +2328,13 @@ void FLevelEditorActionCallbacks::CreateBlankBlueprintClass()
 
 bool FLevelEditorActionCallbacks::CanConvertSelectedActorsIntoBlueprintClass()
 {
-	return GEditor->GetSelectedActorCount() > 0;
+	return (FCreateBlueprintFromActorDialog::GetValidCreationMethods() != ECreateBlueprintFromActorMode::None);
 }
 
 void FLevelEditorActionCallbacks::ConvertSelectedActorsIntoBlueprintClass()
 {
-	const ECreateBlueprintFromActorMode DefaultCreateMode = (GEditor->GetSelectedActorCount() == 1 ? ECreateBlueprintFromActorMode::Subclass : ECreateBlueprintFromActorMode::ChildActor);
+	const ECreateBlueprintFromActorMode ValidCreateModes = FCreateBlueprintFromActorDialog::GetValidCreationMethods();
+	const ECreateBlueprintFromActorMode DefaultCreateMode = ((GEditor->GetSelectedActorCount() == 1 && !!(ValidCreateModes & ECreateBlueprintFromActorMode::Subclass)) ? ECreateBlueprintFromActorMode::Subclass : ECreateBlueprintFromActorMode::ChildActor);
 	FCreateBlueprintFromActorDialog::OpenDialog(DefaultCreateMode);
 }
 
