Fix default selection when bringing up create blueprint from actor for a single selection to be the class of the selected object.
Prevent use of an invalid parent class when creating a blueprint from actor via subclass.
Update the filtered class list when changing from subclass to harvest/child actor mode when a single actor is selected.
#rnx

diff --git a/Engine/Source/Editor/KismetWidgets/Private/CreateBlueprintFromActorDialog.cpp b/Engine/Source/Editor/KismetWidgets/Private/CreateBlueprintFromActorDialog.cpp
--- a/Engine/Source/Editor/KismetWidgets/Private/CreateBlueprintFromActorDialog.cpp
+++ b/Engine/Source/Editor/KismetWidgets/Private/CreateBlueprintFromActorDialog.cpp
@@ -34,8 +34,6 @@
 
 #define LOCTEXT_NAMESPACE "CreateBlueprintFromActorDialog"
 
-TWeakObjectPtr<AActor> FCreateBlueprintFromActorDialog::ActorOverride = nullptr;
-
 class SSCreateBlueprintPicker : public SCompoundWidget
 {
 public:
@@ -43,7 +41,7 @@ public:
 	{}
 
 	SLATE_ARGUMENT(TSharedPtr<SWindow>, ParentWindow)
-		SLATE_ARGUMENT(FClassViewerInitializationOptions, Options)
+		SLATE_ARGUMENT(AActor*, ActorOverride)
 		SLATE_ARGUMENT(ECreateBlueprintFromActorMode, CreateMode)
 		SLATE_END_ARGS()
 
@@ -84,6 +82,7 @@ public:
 		if (NewCheckedState == ECheckBoxState::Checked) 
 		{ 
 			CreateMode = InCreateMode; 
+			ClassViewer->Refresh();
 		} 	
 	};
 
@@ -109,6 +108,9 @@ public:
 	/** The class that was last clicked on */
 	UClass* ChosenClass;
 
+	/** The actor that was passed in */
+	TWeakObjectPtr<AActor> ActorOverride;
+
 	/** The path the asset should be created at */
 	FString AssetPath;
 
@@ -233,7 +235,66 @@ void SSCreateBlueprintPicker::Construct(const FArguments& InArgs)
 	FClassViewerModule& ClassViewerModule = FModuleManager::LoadModuleChecked<FClassViewerModule>("ClassViewer");
 	FAssetToolsModule& AssetToolsModule = FModuleManager::LoadModuleChecked<FAssetToolsModule>("AssetTools");
 
-	ClassViewer = StaticCastSharedRef<SClassViewer>(ClassViewerModule.CreateClassViewer(InArgs._Options, FOnClassPicked::CreateSP(this, &SSCreateBlueprintPicker::OnClassPicked)));
+	ActorOverride = InArgs._ActorOverride;
+
+	FClassViewerInitializationOptions ClassViewerOptions;
+	ClassViewerOptions.Mode = EClassViewerMode::ClassPicker;
+	ClassViewerOptions.DisplayMode = EClassViewerDisplayMode::TreeView;
+	ClassViewerOptions.bShowObjectRootClass = true;
+	ClassViewerOptions.bIsPlaceableOnly = true;
+	ClassViewerOptions.bIsBlueprintBaseOnly = true;
+	ClassViewerOptions.bShowUnloadedBlueprints = true;
+	ClassViewerOptions.bEnableClassDynamicLoading = true;
+	ClassViewerOptions.NameTypeToDisplay = EClassViewerNameTypeToDisplay::Dynamic;
+
+	if (ActorOverride == nullptr)
+	{
+		USelection* SelectedActors = GEditor->GetSelectedActors();
+		if (SelectedActors->Num() == 1)
+		{
+			ActorOverride = CastChecked<AActor>(SelectedActors->GetSelectedObject(0));
+		}
+	}
+
+	class FBlueprintFromActorParentFilter : public IClassViewerFilter
+	{
+	public:
+		FBlueprintFromActorParentFilter(UClass* InAllowedClass, ECreateBlueprintFromActorMode& InCreateModeRef)
+			: CreateModeRef(InCreateModeRef)
+		{
+			AllowedClass.Add(InAllowedClass);
+		}
+
+		TSet<const UClass*> AllowedClass;
+		const ECreateBlueprintFromActorMode& CreateModeRef;
+
+		virtual bool IsClassAllowed(const FClassViewerInitializationOptions& InInitOptions, const UClass* InClass, TSharedRef< FClassViewerFilterFuncs > InFilterFuncs) override
+		{
+			return CreateModeRef != ECreateBlueprintFromActorMode::Subclass || InFilterFuncs->IfInChildOfClassesSet(AllowedClass, InClass) == EFilterReturn::Passed;
+		}
+
+		virtual bool IsUnloadedClassAllowed(const FClassViewerInitializationOptions& InInitOptions, const TSharedRef< const IUnloadedBlueprintData > InUnloadedClassData, TSharedRef< FClassViewerFilterFuncs > InFilterFuncs) override
+		{
+			return CreateModeRef != ECreateBlueprintFromActorMode::Subclass || InFilterFuncs->IfInChildOfClassesSet(AllowedClass, InUnloadedClassData) == EFilterReturn::Passed;
+		}
+	};
+
+	if (ActorOverride.IsValid())
+	{
+		TSharedPtr<FBlueprintFromActorParentFilter> Filter = MakeShareable(new FBlueprintFromActorParentFilter(ActorOverride->GetClass(), CreateMode));
+		ClassViewerOptions.ClassFilter = Filter;
+	}
+
+	if (CreateMode == ECreateBlueprintFromActorMode::Subclass)
+	{
+		ClassViewerOptions.InitiallySelectedClass = ActorOverride->GetClass();
+	}
+	else
+	{
+		ClassViewerOptions.InitiallySelectedClass = AActor::StaticClass();
+	}
+
+	ClassViewer = StaticCastSharedRef<SClassViewer>(ClassViewerModule.CreateClassViewer(ClassViewerOptions, FOnClassPicked::CreateSP(this, &SSCreateBlueprintPicker::OnClassPicked)));
 
 	FString PackageName;
 	AssetPath = ContentBrowserModule.Get().GetCurrentPath();
@@ -656,10 +717,19 @@ void SSCreateBlueprintPicker::UpdateFilenameStatus()
 
 EVisibility SSCreateBlueprintPicker::GetSelectButtonVisibility() const
 {
-	EVisibility ButtonVisibility = EVisibility::Hidden;
-	if (ChosenClass != nullptr && !bIsReportingError)
+	EVisibility ButtonVisibility = EVisibility::Visible;
+	if (ChosenClass == nullptr || bIsReportingError)
+	{
+		ButtonVisibility = EVisibility::Hidden;
+	}
+	else if (CreateMode == ECreateBlueprintFromActorMode::Subclass)
 	{
-		ButtonVisibility = EVisibility::Visible;
+		UObject* SelectedActor = GEditor->GetSelectedActors()->GetSelectedObject(0);
+
+		if (SelectedActor == nullptr || !ChosenClass->IsChildOf(SelectedActor->GetClass()))
+		{
+			ButtonVisibility = EVisibility::Hidden;
+		}
 	}
 	return ButtonVisibility;
 }
@@ -686,46 +756,7 @@ FReply SSCreateBlueprintPicker::OnKeyDown(const FGeometry& MyGeometry, const FKe
 
 void FCreateBlueprintFromActorDialog::OpenDialog(ECreateBlueprintFromActorMode CreateMode, AActor* InActorOverride )
 {
-	ActorOverride = InActorOverride;
-
-	FClassViewerInitializationOptions ClassViewerOptions;
-	ClassViewerOptions.Mode = EClassViewerMode::ClassPicker;
-	ClassViewerOptions.DisplayMode = EClassViewerDisplayMode::TreeView;
-	ClassViewerOptions.bShowObjectRootClass = true;
-	ClassViewerOptions.bIsPlaceableOnly = true;
-	ClassViewerOptions.bIsBlueprintBaseOnly = true;
-	ClassViewerOptions.bShowUnloadedBlueprints = true;
-	ClassViewerOptions.bEnableClassDynamicLoading = true;
-	ClassViewerOptions.NameTypeToDisplay = EClassViewerNameTypeToDisplay::Dynamic;
-
-	if (InActorOverride)
-	{
-		class FBlueprintFromActorParentFilter : public IClassViewerFilter
-		{
-		public:
-			TSet<const UClass*> AllowedClass;
-
-			virtual bool IsClassAllowed(const FClassViewerInitializationOptions& InInitOptions, const UClass* InClass, TSharedRef< FClassViewerFilterFuncs > InFilterFuncs) override
-			{
-				return InFilterFuncs->IfInChildOfClassesSet(AllowedClass, InClass) == EFilterReturn::Passed;
-			}
-
-			virtual bool IsUnloadedClassAllowed(const FClassViewerInitializationOptions& InInitOptions, const TSharedRef< const IUnloadedBlueprintData > InUnloadedClassData, TSharedRef< FClassViewerFilterFuncs > InFilterFuncs) override
-			{
-				return InFilterFuncs->IfInChildOfClassesSet(AllowedClass, InUnloadedClassData) == EFilterReturn::Passed;
-			}
-		};
-
-		TSharedPtr<FBlueprintFromActorParentFilter> Filter = MakeShareable(new FBlueprintFromActorParentFilter);
-		ClassViewerOptions.ClassFilter = Filter;
-		Filter->AllowedClass.Add(InActorOverride->GetClass());
-
-		ClassViewerOptions.InitiallySelectedClass = InActorOverride->GetClass();
-	}
-	else
-	{
-		ClassViewerOptions.InitiallySelectedClass = AActor::StaticClass();
-	}
+	TWeakObjectPtr<AActor> ActorOverride(InActorOverride);
 
 	// Create the window to pick the class
 	TSharedRef<SWindow> PickerWindow = SNew(SWindow)
@@ -737,23 +768,22 @@ void FCreateBlueprintFromActorDialog::OpenDialog(ECreateBlueprintFromActorMode C
 
 	TSharedRef<SSCreateBlueprintPicker> ClassPickerDialog = SNew(SSCreateBlueprintPicker)
 		.ParentWindow(PickerWindow)
-		.Options(ClassViewerOptions)
+		.ActorOverride(InActorOverride)
 		.CreateMode(CreateMode);
 
 	PickerWindow->SetContent(ClassPickerDialog);
-
-//	FSlateApplication::Get().AddWindow(PickerWindow);
+		
 	GEditor->EditorAddModalWindow(PickerWindow);
 
 	if (ClassPickerDialog->bPressedOk)
 	{
 		FString NewAssetName = ClassPickerDialog->AssetPath / ClassPickerDialog->AssetName;
 
-		OnCreateBlueprint(NewAssetName, ClassPickerDialog->ChosenClass, ClassPickerDialog->CreateMode);
+		OnCreateBlueprint(NewAssetName, ClassPickerDialog->ChosenClass, ClassPickerDialog->CreateMode, ActorOverride.Get());
 	}
 }
 
-void FCreateBlueprintFromActorDialog::OnCreateBlueprint(const FString& InAssetPath, UClass* ParentClass, ECreateBlueprintFromActorMode CreateMode)
+void FCreateBlueprintFromActorDialog::OnCreateBlueprint(const FString& InAssetPath, UClass* ParentClass, ECreateBlueprintFromActorMode CreateMode, AActor* ActorToUse)
 {
 	UBlueprint* Blueprint = nullptr;
 
@@ -780,11 +810,9 @@ void FCreateBlueprintFromActorDialog::OnCreateBlueprint(const FString& InAssetPa
 
 		case ECreateBlueprintFromActorMode::Subclass:
 		{
-			AActor* ActorToUse = ActorOverride.Get();
-
 			if (!ActorToUse)
 			{
-				TArray< UObject* > SelectedActors;
+				TArray<UObject*> SelectedActors;
 				GEditor->GetSelectedActors()->GetSelectedObjects(AActor::StaticClass(), SelectedActors);
 				check(SelectedActors.Num() == 1);
 				ActorToUse = Cast<AActor>(SelectedActors[0]);
diff --git a/Engine/Source/Editor/KismetWidgets/Public/CreateBlueprintFromActorDialog.h b/Engine/Source/Editor/KismetWidgets/Public/CreateBlueprintFromActorDialog.h
--- a/Engine/Source/Editor/KismetWidgets/Public/CreateBlueprintFromActorDialog.h
+++ b/Engine/Source/Editor/KismetWidgets/Public/CreateBlueprintFromActorDialog.h
@@ -50,7 +50,5 @@ private:
 	 * @param InAssetPath		Path of the asset to create
 	 * @param bInHarvest		The mode to use when creating a blueprint from the selected actors
 	 */
-	static void OnCreateBlueprint(const FString& InAssetPath, UClass* ParentClass, ECreateBlueprintFromActorMode CreateMode);
-
-	static TWeakObjectPtr<AActor> ActorOverride;
+	static void OnCreateBlueprint(const FString& InAssetPath, UClass* ParentClass, ECreateBlueprintFromActorMode CreateMode, AActor* ActorToUse);
 };
