Recreate the variable type selector combo box when switching between key/value for Map variables

diff --git a/Engine/Source/Editor/KismetWidgets/Private/SPinTypeSelector.cpp b/Engine/Source/Editor/KismetWidgets/Private/SPinTypeSelector.cpp
--- a/Engine/Source/Editor/KismetWidgets/Private/SPinTypeSelector.cpp
+++ b/Engine/Source/Editor/KismetWidgets/Private/SPinTypeSelector.cpp
@@ -907,8 +907,9 @@ TSharedRef<SWidget>	SPinTypeSelector::GetMenuContent(bool bForSecondaryType)
 
 	FilteredTypeTreeRoot = TypeTreeRoot;
 
-	if( !MenuContent.IsValid() )
+	if( !MenuContent.IsValid() || (bForSecondaryType != bMenuContentIsSecondary) )
 	{
+		bMenuContentIsSecondary = bForSecondaryType;
 		// Pre-build the tree view and search box as it is needed as a parameter for the context menu's container.
 		SAssignNew(TypeTreeView, SPinTypeTreeView)
 			.TreeItemsSource(&FilteredTypeTreeRoot)
@@ -947,11 +948,17 @@ TSharedRef<SWidget>	SPinTypeSelector::GetMenuContent(bool bForSecondaryType)
 			];
 			
 
-			TypeComboButton->SetMenuContentWidgetToFocus(FilterTextBox);
+		if (bForSecondaryType)
+		{
 			if (SecondaryTypeComboButton.IsValid())
 			{
 				SecondaryTypeComboButton->SetMenuContentWidgetToFocus(FilterTextBox);
 			}
+		}
+		else
+		{
+			TypeComboButton->SetMenuContentWidgetToFocus(FilterTextBox);
+		}
 	}
 	else
 	{
diff --git a/Engine/Source/Editor/KismetWidgets/Public/SPinTypeSelector.h b/Engine/Source/Editor/KismetWidgets/Public/SPinTypeSelector.h
--- a/Engine/Source/Editor/KismetWidgets/Public/SPinTypeSelector.h
+++ b/Engine/Source/Editor/KismetWidgets/Public/SPinTypeSelector.h
@@ -133,6 +133,9 @@ protected:
 	/** TRUE when the right mouse button is pressed, keeps from handling a right click that does not begin in the widget */
 	bool bIsRightMousePressed;
 
+	/** true if GetMenuContent was last called with bForSecondaryType == true */
+	bool bMenuContentIsSecondary = false;
+
 	/** Whether the selector is using the compact or full mode, or not a selector at all, but just the type image.*/
 	ESelectorType SelectorType;
 
