Fix rare crash when selected actor becomes invalidated.
#rnx

diff --git a/Engine/Source/Editor/KismetWidgets/Private/CreateBlueprintFromActorDialog.cpp b/Engine/Source/Editor/KismetWidgets/Private/CreateBlueprintFromActorDialog.cpp
--- a/Engine/Source/Editor/KismetWidgets/Private/CreateBlueprintFromActorDialog.cpp
+++ b/Engine/Source/Editor/KismetWidgets/Private/CreateBlueprintFromActorDialog.cpp
@@ -279,15 +279,17 @@ void SSCreateBlueprintPicker::Construct(const FArguments& InArgs)
 		}
 	};
 
+	UClass* ActorOverrideClass = nullptr;
 	if (ActorOverride.IsValid())
 	{
-		TSharedPtr<FBlueprintFromActorParentFilter> Filter = MakeShareable(new FBlueprintFromActorParentFilter(ActorOverride->GetClass(), CreateMode));
+		ActorOverrideClass = ActorOverride->GetClass();
+		TSharedPtr<FBlueprintFromActorParentFilter> Filter = MakeShareable(new FBlueprintFromActorParentFilter(ActorOverrideClass, CreateMode));
 		ClassViewerOptions.ClassFilter = Filter;
 	}
 
-	if (CreateMode == ECreateBlueprintFromActorMode::Subclass)
+	if (ActorOverrideClass && CreateMode == ECreateBlueprintFromActorMode::Subclass)
 	{
-		ClassViewerOptions.InitiallySelectedClass = ActorOverride->GetClass();
+		ClassViewerOptions.InitiallySelectedClass = ActorOverrideClass;
 	}
 	else
 	{
