Display the native C++ type of UScriptStructs in the blueprint variable selection tooltip. This will help with structs that use the Display name meta tag, which could possibly result in multiple structs of the same name

diff --git a/Engine/Source/Editor/KismetWidgets/Private/SPinTypeSelector.cpp b/Engine/Source/Editor/KismetWidgets/Private/SPinTypeSelector.cpp
--- a/Engine/Source/Editor/KismetWidgets/Private/SPinTypeSelector.cpp
+++ b/Engine/Source/Editor/KismetWidgets/Private/SPinTypeSelector.cpp
@@ -564,7 +564,20 @@ TSharedRef<ITableRow> SPinTypeSelector::GenerateTypeTreeRow(FPinTypeTreeItem InI
 
 	// Use tooltip if supplied, otherwise just repeat description
 	const FText OrgTooltip = InItem->GetToolTip();
-	const FText Tooltip = !OrgTooltip.IsEmpty() ? OrgTooltip : Description;
+	FText Tooltip = !OrgTooltip.IsEmpty() ? OrgTooltip : Description;
+
+	// If this is a struct type, get some useful information about it's native C++ declaration
+	if (PinType.PinCategory == UEdGraphSchema_K2::PC_Struct)
+	{
+		UScriptStruct* StructType = Cast<UScriptStruct>(PinType.PinSubCategoryObject);
+		if (StructType && StructType->IsNative())
+		{
+			Tooltip = FText::Format(LOCTEXT("NativePinTypeName", "{0}\n\n@see {1}"),
+				Tooltip,
+				FText::FromString(StructType->GetStructCPPName())
+			);
+		}
+	}
 
 	const FString PinTooltipExcerpt = ((PinType.PinCategory != UEdGraphSchema_K2::PC_Byte || PinType.PinSubCategoryObject == nullptr) ? PinType.PinCategory.ToString() : TEXT("Enum")); 
 
