Miscellanious reskin tweaks to the MyBlueprint panel in BP editors.

diff --git a/Engine/Source/Editor/EditorStyle/Private/StarshipStyle.cpp b/Engine/Source/Editor/EditorStyle/Private/StarshipStyle.cpp
--- a/Engine/Source/Editor/EditorStyle/Private/StarshipStyle.cpp
+++ b/Engine/Source/Editor/EditorStyle/Private/StarshipStyle.cpp
@@ -5243,6 +5243,17 @@ void FStarshipEditorStyle::FStyle::SetupPersonaStyle()
 			.SetNormalPadding ( FMargin( 0,0,0,0 ) )
 			.SetPressedPadding( FMargin( 1,1,2,2 ) )
 			);
+
+		const FComboButtonStyle& SimpleComboButton = FStarshipCoreStyle::GetCoreStyle().GetWidgetStyle<FComboButtonStyle>("SimpleComboButton");
+		
+		Set("BlueprintEditor.CompactVariableTypeSelector", 
+			FComboButtonStyle(SimpleComboButton)
+			.SetDownArrowPadding(0)
+			.SetButtonStyle(
+				FButtonStyle(SimpleComboButton.ButtonStyle)
+				.SetNormalPadding(FMargin(0, 2, 2, 2))
+				.SetPressedPadding(FMargin(0, 3, 2, 1)))
+		);
 	}
 
 	// Access modifier display in MyBlueprint window for functions/variables
diff --git a/Engine/Source/Editor/GraphEditor/Private/KismetPins/SGraphPinBool.cpp b/Engine/Source/Editor/GraphEditor/Private/KismetPins/SGraphPinBool.cpp
--- a/Engine/Source/Editor/GraphEditor/Private/KismetPins/SGraphPinBool.cpp
+++ b/Engine/Source/Editor/GraphEditor/Private/KismetPins/SGraphPinBool.cpp
@@ -13,7 +13,6 @@ void SGraphPinBool::Construct(const FArguments& InArgs, UEdGraphPin* InGraphPinO
 TSharedRef<SWidget>	SGraphPinBool::GetDefaultValueWidget()
 {
 	return SNew(SCheckBox)
-		.Style(FEditorStyle::Get(), "Graph.Checkbox")
 		.IsChecked(this, &SGraphPinBool::IsDefaultValueChecked)
 		.IsEnabled(this, &SGraphPin::GetDefaultValueIsEditable)
 		.OnCheckStateChanged(this, &SGraphPinBool::OnDefaultValueCheckBoxChanged)
diff --git a/Engine/Source/Editor/Kismet/Private/SBlueprintPalette.cpp b/Engine/Source/Editor/Kismet/Private/SBlueprintPalette.cpp
--- a/Engine/Source/Editor/Kismet/Private/SBlueprintPalette.cpp
+++ b/Engine/Source/Editor/Kismet/Private/SBlueprintPalette.cpp
@@ -694,6 +694,7 @@ class SPinTypeSelectorHelper : public SCompoundWidget
 {
 public:
 	SLATE_BEGIN_ARGS( SPinTypeSelectorHelper ) {}
+		SLATE_ATTRIBUTE(bool, ReadOnly)
 	SLATE_END_ARGS()
 
 	/**
@@ -714,11 +715,12 @@ public:
 		this->ChildSlot
 		[
 			SNew(SPinTypeSelector, FGetPinTypeTree::CreateUObject(Schema, &UEdGraphSchema_K2::GetVariableTypeTree))
+			.ReadOnly(InArgs._ReadOnly)
 			.Schema(Schema)
 			.TargetPinType(this, &SPinTypeSelectorHelper::OnGetVarType)
 			.OnPinTypeChanged(this, &SPinTypeSelectorHelper::OnVarTypeChanged)
 			.TypeTreeFilter(ETypeTreeFilter::None)
-			.SelectorType(BlueprintEditorPtr.IsValid() ? SPinTypeSelector::ESelectorType::Compact : SPinTypeSelector::ESelectorType::None)
+			.SelectorType(BlueprintEditorPtr.IsValid() ? SPinTypeSelector::ESelectorType::Partial : SPinTypeSelector::ESelectorType::None)
 		];
 	}
 
@@ -1079,7 +1081,8 @@ void SBlueprintPaletteItem::Construct(const FArguments& InArgs, FCreateWidgetFor
 			{
 				const UEdGraphSchema_K2* Schema = GetDefault<UEdGraphSchema_K2>();
 				IconWidget = SNew(SPinTypeSelectorHelper, VariableProp, Blueprint, BlueprintEditorPtr)
-					.IsEnabled(bIsEditingEnabled);
+					.IsEnabled(bIsEditingEnabled)
+					.ReadOnly_Lambda([this]() {return !IsHovered(); });
 			}
 		}
 	};
@@ -1154,53 +1157,84 @@ void SBlueprintPaletteItem::Construct(const FArguments& InArgs, FCreateWidgetFor
 	// Create the widget with an icon
 	TSharedRef<SHorizontalBox> ActionBox = SNew(SHorizontalBox)		
 		.AddMetaData<FTutorialMetaData>(TagMeta);
+	
+	if (GraphAction->GetTypeId() == FEdGraphSchemaAction_K2Var::StaticGetTypeId() || GraphAction->GetTypeId() == FEdGraphSchemaAction_K2LocalVar::StaticGetTypeId())
+	{
+		ActionBox.Get().AddSlot()
+			.FillWidth(0.6f)
+			.VAlign(VAlign_Center)
+			.Padding(3.0f, 0.0f)
+			[
+				NameSlotWidget
+			];
 
-	ActionBox.Get().AddSlot()
-		.AutoWidth()
-		.Padding(FMargin(0.0f, 0.0f, 3.0f, 0.0f))
-		.HAlign(HAlign_Left)
-		.VAlign(VAlign_Center)
-		[
-			SNew(SPaletteItemVisibilityToggle, ActionPtr, InBlueprintEditor, InBlueprint)
-			.IsEnabled(bIsEditingEnabled)
-		];
+		ActionBox.Get().AddSlot()
+			.FillWidth(0.4f)
+			.HAlign(HAlign_Left)
+			.VAlign(VAlign_Center)
+			[
+				IconWidget
+			];
 
+		ActionBox.Get().AddSlot()
+			.AutoWidth()
+			.Padding(FMargin(6.0f, 0.0f, 3.0f, 0.0f))
+			.HAlign(HAlign_Right)
+			.VAlign(VAlign_Center)
+			[
+				SNew(SPaletteItemVisibilityToggle, ActionPtr, InBlueprintEditor, InBlueprint)
+				.IsEnabled(bIsEditingEnabled)
+			];
+	}
+	else
+	{
+		ActionBox.Get().AddSlot()
+			.AutoWidth()
+			.Padding(FMargin(0.0f, 0.0f, 3.0f, 0.0f))
+			.HAlign(HAlign_Left)
+			.VAlign(VAlign_Center)
+			[
+				SNew(SPaletteItemVisibilityToggle, ActionPtr, InBlueprintEditor, InBlueprint)
+				.IsEnabled(bIsEditingEnabled)
+			];
 
-	ActionBox.Get().AddSlot()
-		.AutoWidth()
-		.VAlign(VAlign_Center)
-		[
-			IconWidget
-		];
 
-	// Only add an access specifier if we have one
-	if (ActionAccessSpecifier != EAccessSpecifier::None)
-	{
 		ActionBox.Get().AddSlot()
-			.MaxWidth(50.f)
-			.FillWidth(AccessSpecifierEnabled ? 0.4f : 0.0f)
-			.Padding(FMargin(/* horizontal */ AccessSpecifierEnabled ? 6.0f : 0.0f, /* vertical */ 0.0f))
+			.AutoWidth()
 			.VAlign(VAlign_Center)
-			.HAlign(HAlign_Right)
 			[
-				SNew(STextBlock)
-				// Will only display text if we have a modifier level
+				IconWidget
+			];
+
+		// Only add an access specifier if we have one
+		if (ActionAccessSpecifier != EAccessSpecifier::None)
+		{
+			ActionBox.Get().AddSlot()
+				.MaxWidth(50.f)
+				.FillWidth(AccessSpecifierEnabled ? 0.4f : 0.0f)
+				.Padding(FMargin(/* horizontal */ AccessSpecifierEnabled ? 6.0f : 0.0f, /* vertical */ 0.0f))
+				.VAlign(VAlign_Center)
+				.HAlign(HAlign_Right)
+				[
+					SNew(STextBlock)
+					// Will only display text if we have a modifier level
 					.IsEnabled(AccessSpecifierEnabled)
 					.Text(AccessModifierText)
 					.ColorAndOpacity(FSlateColor::UseSubduedForeground())
 					// Bold if public
 					.TextStyle(FEditorStyle::Get(), ActionAccessSpecifier == EAccessSpecifier::Public ? "BlueprintEditor.AccessModifier.Public" : "BlueprintEditor.AccessModifier.Default")
+				];
+		}
+
+		ActionBox.Get().AddSlot()
+			.FillWidth(1.f)
+			.VAlign(VAlign_Center)
+			.Padding(/* horizontal */ 3.0f, /* vertical */ 3.0f)
+			[
+				NameSlotWidget
 			];
 	}
 
-	ActionBox.Get().AddSlot()
-		.FillWidth(1.f)
-		.VAlign(VAlign_Center)
-		.Padding(/* horizontal */ 3.0f, /* vertical */ 0.0f)
-		[
-			NameSlotWidget
-		];
-
 	// Now, create the actual widget
 	ChildSlot
 	[
diff --git a/Engine/Source/Editor/Kismet/Private/SSCSEditor.cpp b/Engine/Source/Editor/Kismet/Private/SSCSEditor.cpp
--- a/Engine/Source/Editor/Kismet/Private/SSCSEditor.cpp
+++ b/Engine/Source/Editor/Kismet/Private/SSCSEditor.cpp
@@ -4042,7 +4042,7 @@ void SSCSEditor::Construct( const FArguments& InArgs )
 		]
 
 		+ SVerticalBox::Slot()
-		.Padding(0.0f, 0.0f)
+		.Padding(4.0f, 0.0f)
 		[
 			SNew(SBox)
 			.AddMetaData<FTagMetaData>(FTagMetaData(TEXT("ComponentsPanel")))
diff --git a/Engine/Source/Editor/KismetWidgets/Private/SPinTypeSelector.cpp b/Engine/Source/Editor/KismetWidgets/Private/SPinTypeSelector.cpp
--- a/Engine/Source/Editor/KismetWidgets/Private/SPinTypeSelector.cpp
+++ b/Engine/Source/Editor/KismetWidgets/Private/SPinTypeSelector.cpp
@@ -17,6 +17,7 @@
 #include "Kismet2/BlueprintEditorUtils.h"
 #include "Framework/Notifications/NotificationManager.h"
 #include "Widgets/Notifications/SNotificationList.h"
+#include "Widgets/Layout/SWidgetSwitcher.h"
 
 #define LOCTEXT_NAMESPACE "PinTypeSelector"
 
@@ -212,6 +213,8 @@ void SPinTypeSelector::Construct(const FArguments& InArgs, FGetPinTypeTree GetPi
 
 	SearchText = FText::GetEmpty();
 
+	ReadOnly = InArgs._ReadOnly;
+
 	OnTypeChanged = InArgs._OnPinTypeChanged;
 	OnTypePreChanged = InArgs._OnPinTypePreChanged;
 
@@ -262,168 +265,207 @@ void SPinTypeSelector::Construct(const FArguments& InArgs, FGetPinTypeTree GetPi
 				.Image(this, &SPinTypeSelector::GetTypeIconImage)
 				.ColorAndOpacity(this, &SPinTypeSelector::GetTypeIconColor);
 	}
-	else if (SelectorType == ESelectorType::Full)
+	else if (SelectorType == ESelectorType::Full || SelectorType == ESelectorType::Partial)
 	{
-		// Traditional Pin Type Selector with a combo button, the icon, the current type name, and a toggle button for being an array
-		TSharedPtr<SWidget> ContainerControl = SNew(SComboButton)
-		.ButtonStyle(FCoreStyle::Get(), "NoBorder")
-		.HasDownArrow(false)
-		.MenuPlacement(EMenuPlacement::MenuPlacement_ComboBoxRight)
-		.OnGetMenuContent(
-			FOnGetContent::CreateLambda(
-				[this]()
-				{
-					typedef SListView< TSharedPtr<EPinContainerType> > SPinContainerListView;
-					return SNew(SPinContainerListView)
-						.ListItemsSource(&PinTypes)
-						.OnGenerateRow(
-							SPinContainerListView::FOnGenerateRow::CreateLambda(
-								[this](TSharedPtr<EPinContainerType> InPinContainerType, const TSharedRef<STableViewBase>& OwnerTable)->TSharedRef<ITableRow>
-								{
-									EPinContainerType PinContainerType = *InPinContainerType;
-									check(sizeof(Images) / sizeof(*Images) > (int32)PinContainerType);
-									check(sizeof(Tooltips) / sizeof(*Tooltips) > (int32)PinContainerType);
-									const FSlateBrush* SecondaryIcon = PinContainerType == EPinContainerType::Map ? FEditorStyle::GetBrush(TEXT("Kismet.VariableList.MapValueTypeIcon")) : nullptr;
-
-									return SNew(STableRow<TSharedPtr<EPinContainerType>>, OwnerTable)
-										.Content()
-										[
-											SNew(
-												SLayeredImage,
-													SecondaryIcon,
-													TAttribute<FSlateColor>(this, &SPinTypeSelector::GetSecondaryTypeIconColor)
-													)
-											.Image(Images[(int32)PinContainerType])
-											.ToolTip(IDocumentation::Get()->CreateToolTip(Tooltips[(int32)PinContainerType], nullptr, *BigTooltipDocLink, TEXT("Containers")))
-											.ColorAndOpacity(this, &SPinTypeSelector::GetTypeIconColor)
-										]
-										.IsEnabled(
-											TAttribute<bool>::Create(TAttribute<bool>::FGetter::CreateLambda( 
-												[PinContainerType, this]() 
-												{ 
-													return !ContainerRequiresGetTypeHash(PinContainerType) || FBlueprintEditorUtils::HasGetTypeHash( this->TargetPinType.Get() );
-												} ) )
+		TSharedPtr<SWidget> ContainerControl;
+
+		if(SelectorType == ESelectorType::Full)
+		{
+			// Traditional Pin Type Selector with a combo button, the icon, the current type name, and a toggle button for being an array
+			ContainerControl = SNew(SComboButton)
+				.ButtonStyle(FCoreStyle::Get(), "NoBorder")
+				.HasDownArrow(false)
+				.MenuPlacement(EMenuPlacement::MenuPlacement_ComboBoxRight)
+				.OnGetMenuContent(
+					FOnGetContent::CreateLambda(
+						[this]()
+						{
+							typedef SListView< TSharedPtr<EPinContainerType> > SPinContainerListView;
+							return SNew(SPinContainerListView)
+								.ListItemsSource(&PinTypes)
+								.OnGenerateRow(
+									SPinContainerListView::FOnGenerateRow::CreateLambda(
+										[this](TSharedPtr<EPinContainerType> InPinContainerType, const TSharedRef<STableViewBase>& OwnerTable)->TSharedRef<ITableRow>
+										{
+											EPinContainerType PinContainerType = *InPinContainerType;
+											check(sizeof(Images) / sizeof(*Images) > (int32)PinContainerType);
+											check(sizeof(Tooltips) / sizeof(*Tooltips) > (int32)PinContainerType);
+											const FSlateBrush* SecondaryIcon = PinContainerType == EPinContainerType::Map ? FEditorStyle::GetBrush(TEXT("Kismet.VariableList.MapValueTypeIcon")) : nullptr;
+
+											return SNew(STableRow<TSharedPtr<EPinContainerType>>, OwnerTable)
+												.Content()
+												[
+													SNew(SLayeredImage, SecondaryIcon, TAttribute<FSlateColor>(this, &SPinTypeSelector::GetSecondaryTypeIconColor))
+													.Image(Images[(int32)PinContainerType])
+													.ToolTip(IDocumentation::Get()->CreateToolTip(Tooltips[(int32)PinContainerType], nullptr, *BigTooltipDocLink, TEXT("Containers")))
+													.ColorAndOpacity(this, &SPinTypeSelector::GetTypeIconColor)
+												]
+											.IsEnabled(
+												TAttribute<bool>::Create(TAttribute<bool>::FGetter::CreateLambda(
+													[PinContainerType, this]()
+													{
+														return !ContainerRequiresGetTypeHash(PinContainerType) || FBlueprintEditorUtils::HasGetTypeHash(this->TargetPinType.Get());
+													}))
 											);
-								}
-							)
-						)
-						.OnSelectionChanged(
-							SPinContainerListView::FOnSelectionChanged::CreateLambda(
-								[this](TSharedPtr<EPinContainerType> InType, ESelectInfo::Type)
-								{
-									this->OnContainerTypeSelectionChanged(*InType);
-								}
-							)
-						);
-				}
-			)
-		)
-		.ContentPadding(0)
-		.ToolTip(IDocumentation::Get()->CreateToolTip( TAttribute<FText>(this, &SPinTypeSelector::GetToolTipForContainerWidget), NULL, *BigTooltipDocLink, TEXT("Containers")))
-		.IsEnabled( TargetPinType.Get().PinCategory != UEdGraphSchema_K2::PC_Exec )
-		.Visibility(InArgs._bAllowArrays ? EVisibility::Visible : EVisibility::Collapsed)
-		.ButtonContent()
-		[
-			SNew(
-				SLayeredImage,
-				TAttribute<const FSlateBrush*>(this, &SPinTypeSelector::GetSecondaryTypeIconImage),
-				TAttribute<FSlateColor>(this, &SPinTypeSelector::GetSecondaryTypeIconColor)
-			)
-			.Image(this, &SPinTypeSelector::GetTypeIconImage)
-			.ColorAndOpacity(this, &SPinTypeSelector::GetTypeIconColor)
-		];
+										}
+									)
+								)
+								.OnSelectionChanged(
+									SPinContainerListView::FOnSelectionChanged::CreateLambda(
+										[this](TSharedPtr<EPinContainerType> InType, ESelectInfo::Type)
+										{
+											this->OnContainerTypeSelectionChanged(*InType);
+										}
+									)
+								);
+						}
+					)
+				)
+				.ContentPadding(0)
+							.ToolTip(IDocumentation::Get()->CreateToolTip(TAttribute<FText>(this, &SPinTypeSelector::GetToolTipForContainerWidget), NULL, *BigTooltipDocLink, TEXT("Containers")))
+							.IsEnabled(TargetPinType.Get().PinCategory != UEdGraphSchema_K2::PC_Exec)
+							.Visibility(InArgs._bAllowArrays ? EVisibility::Visible : EVisibility::Collapsed)
+							.ButtonContent()
+							[
+								SNew(SLayeredImage, TAttribute<const FSlateBrush*>(this, &SPinTypeSelector::GetSecondaryTypeIconImage), TAttribute<FSlateColor>(this, &SPinTypeSelector::GetSecondaryTypeIconColor))
+								.Image(this, &SPinTypeSelector::GetTypeIconImage)
+								.ColorAndOpacity(this, &SPinTypeSelector::GetTypeIconColor)
+							];
+		}
 
-		Widget = SNew(SHorizontalBox)
-		+SHorizontalBox::Slot()
+		TSharedRef<SHorizontalBox> HBox = SNew(SHorizontalBox).Clipping(EWidgetClipping::ClipToBoundsAlways);
+		Widget = HBox;
+
+		HBox->AddSlot()
+		.HAlign(HAlign_Left)
 		[
 			SNew(SBox)
-			.WidthOverride(100.f)
+			.WidthOverride(SelectorType == ESelectorType::Full ? 125.f : FOptionalSize())
 			[
-				SAssignNew( TypeComboButton, SComboButton )
-				.MenuPlacement(EMenuPlacement::MenuPlacement_ComboBoxRight)
+				SAssignNew(TypeComboButton, SComboButton)
+				.ComboButtonStyle(FAppStyle::Get(), SelectorType == ESelectorType::Full ? "ComboButton" : "BlueprintEditor.CompactVariableTypeSelector")
 				.OnGetMenuContent(this, &SPinTypeSelector::GetMenuContent, false)
 				.ContentPadding(0)
 				.ToolTipText(this, &SPinTypeSelector::GetToolTipForComboBoxType)
+				.ForegroundColor(FSlateColor::UseForeground())
 				.ButtonContent()
 				[
 					SNew(SHorizontalBox)
-					.Clipping(EWidgetClipping::OnDemand)
-						
+					.Clipping(EWidgetClipping::ClipToBoundsAlways)
 					+ SHorizontalBox::Slot()
 					.VAlign(VAlign_Center)
 					.HAlign(HAlign_Left)
+					.Padding(0.0f, 0.0f, 2.0f, 0.0f)
 					.AutoWidth()
 					[
 						SNew(SImage)
-						.Image( this, &SPinTypeSelector::GetTypeIconImage )
-						.ColorAndOpacity( this, &SPinTypeSelector::GetTypeIconColor )
+						.Image(this, &SPinTypeSelector::GetTypeIconImage)
+						.ColorAndOpacity(this, &SPinTypeSelector::GetTypeIconColor)
 					]
-						
 					+ SHorizontalBox::Slot()
+					.Padding(2.0f, 0.0f, 0.0f, 0.0f)
 					.VAlign(VAlign_Center)
 					.HAlign(HAlign_Left)
 					.AutoWidth()
 					[
 						SNew(STextBlock)
-						.Text( this, &SPinTypeSelector::GetTypeDescription )
+						.Text(this, &SPinTypeSelector::GetTypeDescription )
 						.Font(InArgs._Font)
+						.ColorAndOpacity(FSlateColor::UseForeground())
 					]
 				]
 			]
-		]
-
-		+SHorizontalBox::Slot()
-		.AutoWidth()
-			.VAlign(VAlign_Center)
-			.HAlign(HAlign_Center)
-		[
-			ContainerControl.ToSharedRef()
-		]
+		];
 
-		+SHorizontalBox::Slot()
-		[
-			SNew(SBox)
-			.Visibility(
-				TAttribute<EVisibility>::Create(
-					TAttribute<EVisibility>::FGetter::CreateLambda(
-						[this]() {return this->TargetPinType.Get().IsMap() == true ? EVisibility::Visible : EVisibility::Collapsed; }
+		if(SelectorType == ESelectorType::Full)
+		{
+			HBox->AddSlot()
+				.AutoWidth()
+				.VAlign(VAlign_Center)
+				.HAlign(HAlign_Center)
+				[
+					ContainerControl.ToSharedRef()
+				];
+	
+			HBox->AddSlot()
+			[
+				SNew(SBox)
+				.Visibility(
+					TAttribute<EVisibility>::Create(
+						TAttribute<EVisibility>::FGetter::CreateLambda(
+							[this]() {return this->TargetPinType.Get().IsMap() == true ? EVisibility::Visible : EVisibility::Collapsed; }
+						)
 					)
 				)
-			)
-			[
-				SAssignNew( SecondaryTypeComboButton, SComboButton )
-				.OnGetMenuContent(this, &SPinTypeSelector::GetMenuContent, true )
-				.ContentPadding(0)
-				.ToolTipText(this, &SPinTypeSelector::GetToolTipForComboBoxSecondaryType)
-				.ButtonContent()
 				[
-					SNew(SHorizontalBox)
-					.Clipping(EWidgetClipping::OnDemand)
-
-					+SHorizontalBox::Slot()
-					.AutoWidth()
-					.VAlign(VAlign_Center)
-					.HAlign(HAlign_Center)
-					[
-						SNew(SImage)
-						.Image( this, &SPinTypeSelector::GetSecondaryTypeIconImage )
-						.ColorAndOpacity( this, &SPinTypeSelector::GetSecondaryTypeIconColor )
-					]
-					+SHorizontalBox::Slot()
-					.VAlign(VAlign_Center)
-					.HAlign(HAlign_Left)
+					SAssignNew( SecondaryTypeComboButton, SComboButton )
+					.OnGetMenuContent(this, &SPinTypeSelector::GetMenuContent, true )
+					.ContentPadding(0)
+					.ToolTipText(this, &SPinTypeSelector::GetToolTipForComboBoxSecondaryType)
+					.ButtonContent()
 					[
-						SNew(STextBlock)
-						.Text( this, &SPinTypeSelector::GetSecondaryTypeDescription )
-						.Font(InArgs._Font)
+						SNew(SHorizontalBox)
+						.Clipping(EWidgetClipping::OnDemand)
+						+SHorizontalBox::Slot()
+						.AutoWidth()
+						.VAlign(VAlign_Center)
+						.HAlign(HAlign_Left)
+						.Padding(0.0f, 0.0f, 2.0f, 0.0f)
+						[
+							SNew(SImage)
+							.Image( this, &SPinTypeSelector::GetSecondaryTypeIconImage )
+							.ColorAndOpacity( this, &SPinTypeSelector::GetSecondaryTypeIconColor )
+						]
+						+SHorizontalBox::Slot()
+						.VAlign(VAlign_Center)
+						.HAlign(HAlign_Left)
+						.Padding(2.0f, 0.0f, 0.0f, 0.0f)
+						[
+							SNew(STextBlock)
+							.Text( this, &SPinTypeSelector::GetSecondaryTypeDescription )
+							.Font(InArgs._Font)
+						]
 					]
 				]
-			]
-		];
+			];
+		}
 	}
+
+
 	this->ChildSlot
 	[
-		Widget.ToSharedRef()
+		SNew(SWidgetSwitcher)
+		.WidgetIndex_Lambda([this](){return ReadOnly.Get() ? 1 : 0; })
+		+ SWidgetSwitcher::Slot()
+		[
+			Widget.ToSharedRef()
+		]
+		+ SWidgetSwitcher::Slot()
+		[
+			SNew(SHorizontalBox)
+			.Clipping(EWidgetClipping::OnDemand)
+			+ SHorizontalBox::Slot()
+			.VAlign(VAlign_Center)
+			.HAlign(HAlign_Left)
+			.Padding(FMargin(0.0f, 2.0f, 2.0f, 2.0f))
+			.AutoWidth()
+			[
+				SNew(SImage)
+				.Image(this, &SPinTypeSelector::GetTypeIconImage)
+				.ColorAndOpacity(this, &SPinTypeSelector::GetTypeIconColor)
+			]
+			+ SHorizontalBox::Slot()
+			.Padding(2.0f, 2.0f)
+			.VAlign(VAlign_Center)
+			.HAlign(HAlign_Left)
+			.AutoWidth()
+			[
+				SNew(STextBlock)
+				.Text(this, &SPinTypeSelector::GetTypeDescription)
+				.Font(InArgs._Font)
+				.ColorAndOpacity(FSlateColor::UseSubduedForeground())
+			]
+		]	
 	];
 }
 
diff --git a/Engine/Source/Editor/KismetWidgets/Public/SPinTypeSelector.h b/Engine/Source/Editor/KismetWidgets/Public/SPinTypeSelector.h
--- a/Engine/Source/Editor/KismetWidgets/Public/SPinTypeSelector.h
+++ b/Engine/Source/Editor/KismetWidgets/Public/SPinTypeSelector.h
@@ -50,19 +50,24 @@ public:
 	enum class ESelectorType : uint8
 	{
 		None,
+		// Shows only the type selector pill
 		Compact,
+		// Shows the type selector pill and full name of type
+		Partial,
+		// Shows the type selector pill, full name of type, and container selection type
 		Full
 	};
 
-	SLATE_BEGIN_ARGS( SPinTypeSelector )
+	SLATE_BEGIN_ARGS(SPinTypeSelector)
 		: _TargetPinType()
 		, _Schema(nullptr)
 		, _TypeTreeFilter(ETypeTreeFilter::None)
 		, _bAllowArrays(true)
 		, _TreeViewWidth(300.f)
 		, _TreeViewHeight(400.f)
-		, _Font( FEditorStyle::GetFontStyle( TEXT("NormalFont") ) )
-		, _SelectorType( ESelectorType::Full )
+		, _Font(FEditorStyle::GetFontStyle(TEXT("NormalFont")))
+		, _SelectorType(ESelectorType::Full)
+		, _ReadOnly(false)
 		{}
 		SLATE_ATTRIBUTE( FEdGraphPinType, TargetPinType )
 		SLATE_ARGUMENT( const UEdGraphSchema_K2*, Schema )
@@ -74,6 +79,7 @@ public:
 		SLATE_EVENT( FOnPinTypeChanged, OnPinTypeChanged )
 		SLATE_ATTRIBUTE( FSlateFontInfo, Font )
 		SLATE_ARGUMENT( ESelectorType, SelectorType )
+		SLATE_ATTRIBUTE(bool, ReadOnly);
 	SLATE_END_ARGS()
 public:
 	void Construct(const FArguments& InArgs, FGetPinTypeTree GetPinTypeTreeFunc);
@@ -139,6 +145,9 @@ protected:
 	/** Whether the selector is using the compact or full mode, or not a selector at all, but just the type image.*/
 	ESelectorType SelectorType;
 
+	/** Whether or not the type is read only and not editable (implies a different style) */
+	TAttribute<bool> ReadOnly;
+
 	/** Holds a cache of the allowed Object Reference types for the last sub-menu opened. */
 	TArray<FObjectReferenceListItem> AllowedObjectReferenceTypes;
 	TWeakPtr<SListView<FObjectReferenceListItem>> WeakListView;
