Fix crash when selecting new pin types in graph editors. Also cleaned up the styling

diff --git a/Engine/Source/Editor/KismetWidgets/Private/SPinTypeSelector.cpp b/Engine/Source/Editor/KismetWidgets/Private/SPinTypeSelector.cpp
--- a/Engine/Source/Editor/KismetWidgets/Private/SPinTypeSelector.cpp
+++ b/Engine/Source/Editor/KismetWidgets/Private/SPinTypeSelector.cpp
@@ -18,10 +18,10 @@
 #include "Framework/Notifications/NotificationManager.h"
 #include "Widgets/Notifications/SNotificationList.h"
 #include "Widgets/Layout/SWidgetSwitcher.h"
+#include "Framework/MultiBox/MultiBoxBuilder.h"
 
 #define LOCTEXT_NAMESPACE "PinTypeSelector"
 
-static const FString BigTooltipDocLink = TEXT("Shared/Editor/Blueprint/VariableTypes");
 
 /** Manages items in the Object Reference Type list, the sub-menu of the PinTypeSelector */
 struct FObjectReferenceType
@@ -42,6 +42,36 @@ struct FObjectReferenceType
 	{}
 };
 
+namespace PinTypeSelectorStatics
+{
+	static const FString BigTooltipDocLink = TEXT("Shared/Editor/Blueprint/VariableTypes");
+
+	// SComboBox is a bit restrictive:
+	static TArray<TSharedPtr<EPinContainerType>> PinTypes;
+
+	static FName Images[] = {
+		TEXT("Kismet.VariableList.TypeIcon"),
+		TEXT("Kismet.VariableList.ArrayTypeIcon"),
+		TEXT("Kismet.VariableList.SetTypeIcon"),
+		TEXT("Kismet.VariableList.MapKeyTypeIcon"),
+	};
+
+	static const FText Labels[] = {
+		LOCTEXT("SingleVariable", "Single"),
+		LOCTEXT("Array", "Array"),
+		LOCTEXT("Set", "Set"),
+		LOCTEXT("Map", "Map"),
+	};
+
+	static const FText Tooltips[] = {
+		LOCTEXT("SingleVariableTooltip", "Single Variable"),
+		LOCTEXT("ArrayTooltip", "Array"),
+		LOCTEXT("SetTooltip", "Set"),
+		LOCTEXT("MapTooltip", "Map (Dictionary)"),
+	};
+}
+
+
 class SPinTypeRow : public SComboRow<FPinTypeTreeItem>
 {
 public:
@@ -187,30 +217,6 @@ TSharedRef<SWidget> SPinTypeSelector::ConstructPinTypeImage(UEdGraphPin* Pin)
 
 void SPinTypeSelector::Construct(const FArguments& InArgs, FGetPinTypeTree GetPinTypeTreeFunc)
 {
-	// SComboBox is a bit restrictive:
-	static TArray<TSharedPtr<EPinContainerType>> PinTypes;
-	if (PinTypes.Num() == 0)
-	{
-		PinTypes.Add(MakeShareable(new EPinContainerType(EPinContainerType::None)));
-		PinTypes.Add(MakeShareable(new EPinContainerType(EPinContainerType::Array)));
-		PinTypes.Add(MakeShareable(new EPinContainerType(EPinContainerType::Set)));
-		PinTypes.Add(MakeShareable(new EPinContainerType(EPinContainerType::Map)));
-	}
-
-	static const FSlateBrush* Images[] = {
-		FEditorStyle::GetBrush(TEXT("Kismet.VariableList.TypeIcon")),
-		FEditorStyle::GetBrush(TEXT("Kismet.VariableList.ArrayTypeIcon")),
-		FEditorStyle::GetBrush(TEXT("Kismet.VariableList.SetTypeIcon")),
-		FEditorStyle::GetBrush(TEXT("Kismet.VariableList.MapKeyTypeIcon")),
-	};
-
-	static const FText Tooltips[] = {
-		LOCTEXT("SingleVariableTooltip", "Single Variable"),
-		LOCTEXT("ArrayTooltip", "Array"),
-		LOCTEXT("SetTooltip", "Set"),
-		LOCTEXT("MapTooltip", "Map (Dictionary)"),
-	};
-
 	SearchText = FText::GetEmpty();
 
 	ReadOnly = InArgs._ReadOnly;
@@ -273,64 +279,19 @@ void SPinTypeSelector::Construct(const FArguments& InArgs, FGetPinTypeTree GetPi
 		{
 			// Traditional Pin Type Selector with a combo button, the icon, the current type name, and a toggle button for being an array
 			ContainerControl = SNew(SComboButton)
-				.ButtonStyle(FCoreStyle::Get(), "NoBorder")
-				.HasDownArrow(false)
+				.ComboButtonStyle(FAppStyle::Get(),"BlueprintEditor.CompactVariableTypeSelector")
 				.MenuPlacement(EMenuPlacement::MenuPlacement_ComboBoxRight)
-				.OnGetMenuContent(
-					FOnGetContent::CreateLambda(
-						[this]()
-						{
-							typedef SListView< TSharedPtr<EPinContainerType> > SPinContainerListView;
-							return SNew(SPinContainerListView)
-								.ListItemsSource(&PinTypes)
-								.OnGenerateRow(
-									SPinContainerListView::FOnGenerateRow::CreateLambda(
-										[this](TSharedPtr<EPinContainerType> InPinContainerType, const TSharedRef<STableViewBase>& OwnerTable)->TSharedRef<ITableRow>
-										{
-											EPinContainerType PinContainerType = *InPinContainerType;
-											check(sizeof(Images) / sizeof(*Images) > (int32)PinContainerType);
-											check(sizeof(Tooltips) / sizeof(*Tooltips) > (int32)PinContainerType);
-											const FSlateBrush* SecondaryIcon = PinContainerType == EPinContainerType::Map ? FEditorStyle::GetBrush(TEXT("Kismet.VariableList.MapValueTypeIcon")) : nullptr;
-
-											return SNew(STableRow<TSharedPtr<EPinContainerType>>, OwnerTable)
-												.Content()
-												[
-													SNew(SLayeredImage, SecondaryIcon, TAttribute<FSlateColor>(this, &SPinTypeSelector::GetSecondaryTypeIconColor))
-													.Image(Images[(int32)PinContainerType])
-													.ToolTip(IDocumentation::Get()->CreateToolTip(Tooltips[(int32)PinContainerType], nullptr, *BigTooltipDocLink, TEXT("Containers")))
-													.ColorAndOpacity(this, &SPinTypeSelector::GetTypeIconColor)
-												]
-											.IsEnabled(
-												TAttribute<bool>::Create(TAttribute<bool>::FGetter::CreateLambda(
-													[PinContainerType, this]()
-													{
-														return !ContainerRequiresGetTypeHash(PinContainerType) || FBlueprintEditorUtils::HasGetTypeHash(this->TargetPinType.Get());
-													}))
-											);
-										}
-									)
-								)
-								.OnSelectionChanged(
-									SPinContainerListView::FOnSelectionChanged::CreateLambda(
-										[this](TSharedPtr<EPinContainerType> InType, ESelectInfo::Type)
-										{
-											this->OnContainerTypeSelectionChanged(*InType);
-										}
-									)
-								);
-						}
-					)
-				)
+				.OnGetMenuContent(this, &SPinTypeSelector::GetPinContainerTypeMenuContent)
 				.ContentPadding(0)
-							.ToolTip(IDocumentation::Get()->CreateToolTip(TAttribute<FText>(this, &SPinTypeSelector::GetToolTipForContainerWidget), NULL, *BigTooltipDocLink, TEXT("Containers")))
-							.IsEnabled(TargetPinType.Get().PinCategory != UEdGraphSchema_K2::PC_Exec)
-							.Visibility(InArgs._bAllowArrays ? EVisibility::Visible : EVisibility::Collapsed)
-							.ButtonContent()
-							[
-								SNew(SLayeredImage, TAttribute<const FSlateBrush*>(this, &SPinTypeSelector::GetSecondaryTypeIconImage), TAttribute<FSlateColor>(this, &SPinTypeSelector::GetSecondaryTypeIconColor))
-								.Image(this, &SPinTypeSelector::GetTypeIconImage)
-								.ColorAndOpacity(this, &SPinTypeSelector::GetTypeIconColor)
-							];
+				.ToolTip(IDocumentation::Get()->CreateToolTip(TAttribute<FText>(this, &SPinTypeSelector::GetToolTipForContainerWidget), NULL, *PinTypeSelectorStatics::BigTooltipDocLink, TEXT("Containers")))
+				.IsEnabled(TargetPinType.Get().PinCategory != UEdGraphSchema_K2::PC_Exec)
+				.Visibility(InArgs._bAllowArrays ? EVisibility::Visible : EVisibility::Collapsed)
+				.ButtonContent()
+				[
+					SNew(SLayeredImage, TAttribute<const FSlateBrush*>(this, &SPinTypeSelector::GetSecondaryTypeIconImage), TAttribute<FSlateColor>(this, &SPinTypeSelector::GetSecondaryTypeIconColor))
+					.Image(this, &SPinTypeSelector::GetTypeIconImage)
+					.ColorAndOpacity(this, &SPinTypeSelector::GetTypeIconColor)
+				];
 		}
 
 		TSharedRef<SHorizontalBox> HBox = SNew(SHorizontalBox).Clipping(EWidgetClipping::ClipToBoundsAlways);
@@ -597,7 +558,7 @@ TSharedRef<ITableRow> SPinTypeSelector::GenerateTypeTreeRow(FPinTypeTreeItem InI
 
 	TSharedPtr< SHorizontalBox > HorizontalBox;
 	TSharedRef< ITableRow > ReturnWidget = SNew( SPinTypeRow, OwnerTree, MenuContent )
-		.ToolTip( IDocumentation::Get()->CreateToolTip( Tooltip, NULL, *BigTooltipDocLink, PinTooltipExcerpt) )
+		.ToolTip( IDocumentation::Get()->CreateToolTip( Tooltip, NULL, *PinTypeSelectorStatics::BigTooltipDocLink, PinTooltipExcerpt) )
 		.OnGetMenuContent(OnGetContent)
 		[
 			SAssignNew(HorizontalBox, SHorizontalBox)
@@ -644,7 +605,7 @@ TSharedRef<ITableRow> SPinTypeSelector::GenerateTypeTreeRow(FPinTypeTreeItem InI
 TSharedRef<SWidget> SPinTypeSelector::CreateObjectReferenceWidget(FPinTypeTreeItem InItem, FEdGraphPinType& InPinType, const FSlateBrush* InIconBrush, FText InSimpleTooltip) const
 {
 	return SNew(SHorizontalBox)
-		.ToolTip(IDocumentation::Get()->CreateToolTip(InSimpleTooltip, nullptr, *BigTooltipDocLink, InPinType.PinCategory.ToString()))
+		.ToolTip(IDocumentation::Get()->CreateToolTip(InSimpleTooltip, nullptr, *PinTypeSelectorStatics::BigTooltipDocLink, InPinType.PinCategory.ToString()))
 		+SHorizontalBox::Slot()
 		.AutoWidth()
 		.Padding(1.f)
@@ -996,6 +957,51 @@ TSharedRef<SWidget>	SPinTypeSelector::GetMenuContent(bool bForSecondaryType)
 	return MenuContent.ToSharedRef();
 }
 
+TSharedRef<SWidget> SPinTypeSelector::GetPinContainerTypeMenuContent()
+{
+	FMenuBuilder MenuBuilder(true, nullptr);
+
+	if (PinTypeSelectorStatics::PinTypes.Num() == 0)
+	{
+		PinTypeSelectorStatics::PinTypes.Add(MakeShared<EPinContainerType>(EPinContainerType::None));
+		PinTypeSelectorStatics::PinTypes.Add(MakeShared<EPinContainerType>(EPinContainerType::Array));
+		PinTypeSelectorStatics::PinTypes.Add(MakeShared<EPinContainerType>(EPinContainerType::Set));
+		PinTypeSelectorStatics::PinTypes.Add(MakeShared<EPinContainerType>(EPinContainerType::Map));
+	}
+
+	for (auto& PinType : PinTypeSelectorStatics::PinTypes)
+	{
+		EPinContainerType PinContainerType = *PinType;
+		FUIAction Action;
+		Action.ExecuteAction.BindSP(this, &SPinTypeSelector::OnContainerTypeSelectionChanged, PinContainerType);
+		Action.CanExecuteAction.BindLambda([PinContainerType, this]() { return !ContainerRequiresGetTypeHash(PinContainerType) || FBlueprintEditorUtils::HasGetTypeHash(this->TargetPinType.Get()); });
+
+		const FSlateBrush* SecondaryIcon = PinContainerType == EPinContainerType::Map ? FAppStyle::Get().GetBrush(TEXT("Kismet.VariableList.MapValueTypeIcon")) : nullptr;
+
+		TSharedRef<SWidget> Widget =
+			SNew(SHorizontalBox)
+			+ SHorizontalBox::Slot()
+			.AutoWidth()
+			.VAlign(VAlign_Center)
+			[
+				SNew(SLayeredImage, SecondaryIcon, TAttribute<FSlateColor>(this, &SPinTypeSelector::GetSecondaryTypeIconColor))
+				.Image(FAppStyle::Get().GetBrush(PinTypeSelectorStatics::Images[(int32)PinContainerType]))
+				.ToolTip(IDocumentation::Get()->CreateToolTip(PinTypeSelectorStatics::Tooltips[(int32)PinContainerType], nullptr, *PinTypeSelectorStatics::BigTooltipDocLink, TEXT("Containers")))
+				.ColorAndOpacity(this, &SPinTypeSelector::GetTypeIconColor)
+			]
+			+ SHorizontalBox::Slot()
+			.Padding(4.0f,2.0f)
+			[
+				SNew(STextBlock)
+				.Text(PinTypeSelectorStatics::Labels[(int32)PinContainerType])
+			];
+
+		MenuBuilder.AddMenuEntry(Action, Widget);
+	}
+
+	return MenuBuilder.MakeWidget();
+}
+
 //=======================================================================
 // Search Support
 void SPinTypeSelector::OnFilterTextChanged(const FText& NewText)
diff --git a/Engine/Source/Editor/KismetWidgets/Public/SPinTypeSelector.h b/Engine/Source/Editor/KismetWidgets/Public/SPinTypeSelector.h
--- a/Engine/Source/Editor/KismetWidgets/Public/SPinTypeSelector.h
+++ b/Engine/Source/Editor/KismetWidgets/Public/SPinTypeSelector.h
@@ -180,6 +180,7 @@ protected:
 	/** Reference to the menu content that's displayed when the type button is clicked on */
 	TSharedPtr<SMenuOwner> MenuContent;
 	virtual TSharedRef<SWidget>	GetMenuContent(bool bForSecondaryType);
+	TSharedRef<SWidget> GetPinContainerTypeMenuContent();
 
 	/** Type searching support */
 	FText SearchText;
