[NO NON-PII COMMIT MESSAGE]

diff --git a/Engine/Source/Editor/KismetWidgets/Private/SPinTypeSelector.cpp b/Engine/Source/Editor/KismetWidgets/Private/SPinTypeSelector.cpp
--- a/Engine/Source/Editor/KismetWidgets/Private/SPinTypeSelector.cpp
+++ b/Engine/Source/Editor/KismetWidgets/Private/SPinTypeSelector.cpp
@@ -19,6 +19,7 @@
 #include "Widgets/Notifications/SNotificationList.h"
 #include "Widgets/Layout/SWidgetSwitcher.h"
 #include "Framework/MultiBox/MultiBoxBuilder.h"
+#include "BlueprintEditorSettings.h"
 
 #define LOCTEXT_NAMESPACE "PinTypeSelector"
 
@@ -586,6 +587,13 @@ TSharedRef<ITableRow> SPinTypeSelector::GenerateTypeTreeRow(FPinTypeTreeItem InI
 	const FText OrgTooltip = InItem->GetToolTip();
 	FText Tooltip = !OrgTooltip.IsEmpty() ? OrgTooltip : Description;
 
+	// Switch to short tooltip based on settings
+	UClass* ClassType = Cast<UClass>(PinType.PinSubCategoryObject);
+	if (ClassType && ClassType->IsNative() && GetDefault<UBlueprintEditorSettings>()->bShowShortTooltips)
+	{
+		Tooltip = ClassType->GetToolTipText(true);
+	}
+
 	// If this is a struct type, get some useful information about it's native C++ declaration
 	if (PinType.PinCategory == UEdGraphSchema_K2::PC_Struct)
 	{
