[NO NON-PII COMMIT MESSAGE]

diff --git a/Engine/Source/Editor/BlueprintGraph/Classes/K2Node_ClassDynamicCast.h b/Engine/Source/Editor/BlueprintGraph/Classes/K2Node_ClassDynamicCast.h
--- a/Engine/Source/Editor/BlueprintGraph/Classes/K2Node_ClassDynamicCast.h
+++ b/Engine/Source/Editor/BlueprintGraph/Classes/K2Node_ClassDynamicCast.h
@@ -18,6 +18,7 @@ class UK2Node_ClassDynamicCast : public UK2Node_DynamicCast
 	virtual void AllocateDefaultPins() override;
 	virtual FLinearColor GetNodeTitleColor() const override;
 	virtual FText GetNodeTitle(ENodeTitleType::Type TitleType) const override;
+	virtual FText GetTooltipText() const override;
 	//~ End UEdGraphNode Interface
 
 	//~ Begin UK2Node Interface
diff --git a/Engine/Source/Editor/BlueprintGraph/Classes/K2Node_DynamicCast.h b/Engine/Source/Editor/BlueprintGraph/Classes/K2Node_DynamicCast.h
--- a/Engine/Source/Editor/BlueprintGraph/Classes/K2Node_DynamicCast.h
+++ b/Engine/Source/Editor/BlueprintGraph/Classes/K2Node_DynamicCast.h
@@ -25,6 +25,7 @@ class UK2Node_DynamicCast : public UK2Node
 	virtual void AllocateDefaultPins() override;
 	virtual FLinearColor GetNodeTitleColor() const override;
 	virtual FText GetNodeTitle(ENodeTitleType::Type TitleType) const override;
+	virtual FText GetTooltipText() const override;
 	virtual FSlateIcon GetIconAndTint(FLinearColor& OutColor) const override;
 	virtual void GetNodeContextMenuActions(class UToolMenu* Menu, class UGraphNodeContextMenuContext* Context) const override;
 	virtual bool IncludeParentNodeContextMenu() const override { return true; }
diff --git a/Engine/Source/Editor/BlueprintGraph/Private/K2Node_ClassDynamicCast.cpp b/Engine/Source/Editor/BlueprintGraph/Private/K2Node_ClassDynamicCast.cpp
--- a/Engine/Source/Editor/BlueprintGraph/Private/K2Node_ClassDynamicCast.cpp
+++ b/Engine/Source/Editor/BlueprintGraph/Private/K2Node_ClassDynamicCast.cpp
@@ -76,6 +76,20 @@ FText UK2Node_ClassDynamicCast::GetNodeTitle(ENodeTitleType::Type TitleType) con
 	return CachedNodeTitle;
 }
 
+FText UK2Node_ClassDynamicCast::GetTooltipText() const
+{
+	if (TargetType && TargetType->IsChildOf(UInterface::StaticClass()))
+	{
+		return LOCTEXT("CastToInterfaceTooltip", "Tries to access class as an interface it may implement.");
+	}
+	UBlueprint* CastToBP = UBlueprint::GetBlueprintFromClass(TargetType);
+	if (CastToBP)
+	{
+		return LOCTEXT("CastToBPTooltip", "Tries to access class as a blueprint class it may inherit from.\n\nNOTE: This will cause the blueprint to always be loaded, which can be expensive.");
+	}
+	return LOCTEXT("CastToNativeTooltip", "Tries to access class as one it may inherit from.");
+}
+
 UEdGraphPin* UK2Node_ClassDynamicCast::GetCastSourcePin() const
 {
 	UEdGraphPin* Pin = FindPinChecked(FClassDynamicCastHelper::ClassToCastName);
diff --git a/Engine/Source/Editor/BlueprintGraph/Private/K2Node_DynamicCast.cpp b/Engine/Source/Editor/BlueprintGraph/Private/K2Node_DynamicCast.cpp
--- a/Engine/Source/Editor/BlueprintGraph/Private/K2Node_DynamicCast.cpp
+++ b/Engine/Source/Editor/BlueprintGraph/Private/K2Node_DynamicCast.cpp
@@ -117,6 +117,20 @@ FText UK2Node_DynamicCast::GetNodeTitle(ENodeTitleType::Type TitleType) const
 	return CachedNodeTitle;
 }
 
+FText UK2Node_DynamicCast::GetTooltipText() const
+{
+	if (TargetType && TargetType->IsChildOf(UInterface::StaticClass()))
+	{
+		return LOCTEXT("CastToInterfaceTooltip", "Tries to access object as an interface it may implement.");
+	}
+	UBlueprint* CastToBP = UBlueprint::GetBlueprintFromClass(TargetType);
+	if (CastToBP)
+	{
+		return LOCTEXT("CastToBPTooltip", "Tries to access object as a blueprint class it may be an instance of.\n\nNOTE: This will cause the blueprint to always be loaded, which can be expensive.");
+	}
+	return LOCTEXT("CastToNativeTooltip", "Tries to access object as a class it may be an instance of.");
+}
+
 void UK2Node_DynamicCast::GetNodeContextMenuActions(UToolMenu* Menu, UGraphNodeContextMenuContext* Context) const
 {
 	Super::GetNodeContextMenuActions(Menu, Context);
diff --git a/Engine/Source/Editor/KismetWidgets/Private/SPinTypeSelector.cpp b/Engine/Source/Editor/KismetWidgets/Private/SPinTypeSelector.cpp
--- a/Engine/Source/Editor/KismetWidgets/Private/SPinTypeSelector.cpp
+++ b/Engine/Source/Editor/KismetWidgets/Private/SPinTypeSelector.cpp
@@ -743,9 +743,11 @@ TSharedRef< SWidget > SPinTypeSelector::GetAllowedObjectTypes(FPinTypeTreeItem I
 	const FSlateBrush* IconBrush = FBlueprintEditorUtils::GetIconFromPin(PinType);
 
 	FFormatNamedArguments Args;
+	bool bIsNative = false;
 
-	if(PinType.PinSubCategory != UEdGraphSchema_K2::PSC_Bitmask && PinType.PinSubCategoryObject.IsValid())
+	if (PinType.PinSubCategory != UEdGraphSchema_K2::PSC_Bitmask && PinType.PinSubCategoryObject.IsValid())
 	{
+		bIsNative = PinType.PinSubCategoryObject->IsNative() && PinType.PinSubCategoryObject != UBlueprint::StaticClass();
 		Args.Add(TEXT("TypeName"), InItem->GetDescription());
 	}
 
@@ -764,23 +766,35 @@ TSharedRef< SWidget > SPinTypeSelector::GetAllowedObjectTypes(FPinTypeTreeItem I
 	if (PossibleObjectReferenceTypes & static_cast<uint8>(EObjectReferenceType::ClassReference))
 	{
 		PinType.PinCategory = UEdGraphSchema_K2::PC_Class;
-		TSharedRef<SWidget> Widget = CreateObjectReferenceWidget(InItem, PinType, IconBrush, FText::Format(LOCTEXT("ClassTooltip", "Reference a class of type \'{TypeName}\'"), Args));
+		TSharedRef<SWidget> Widget = CreateObjectReferenceWidget(InItem, PinType, IconBrush, FText::Format(LOCTEXT("ClassTooltip", "Reference a class inheriting from type \'{TypeName}\'"), Args));
 		FObjectReferenceListItem ObjectReferenceType = MakeShareable(new FObjectReferenceType(InItem, Widget, PinType.PinCategory));
 		AllowedObjectReferenceTypes.Add(ObjectReferenceType);
 	}
 
 	if (PossibleObjectReferenceTypes & static_cast<uint8>(EObjectReferenceType::SoftObject))
 	{
+		FText FormatString = LOCTEXT("AssetTooltip", "Path to an instanced object of type \'{TypeName}\' which may not be loaded. Can be used to asynchronously load the asset.");
+		if (!bIsNative)
+		{
+			FormatString = LOCTEXT("AssetBPTooltip", "Path to an instanced object of type \'{TypeName}\' which may not be loaded. The blueprint type itself will always be loaded, which can be expensive.");
+		}
+
 		PinType.PinCategory = UEdGraphSchema_K2::PC_SoftObject;
-		TSharedRef<SWidget> Widget = CreateObjectReferenceWidget(InItem, PinType, IconBrush, FText::Format(LOCTEXT("AssetTooltip", "Path to an instanced object of type \'{TypeName}\' which may be in an unloaded state. Can be utilized to asynchronously load the object reference."), Args));
+		TSharedRef<SWidget> Widget = CreateObjectReferenceWidget(InItem, PinType, IconBrush, FText::Format(FormatString, Args));
 		FObjectReferenceListItem ObjectReferenceType = MakeShareable(new FObjectReferenceType(InItem, Widget, PinType.PinCategory));
 		AllowedObjectReferenceTypes.Add(ObjectReferenceType);
 	}
 
 	if (PossibleObjectReferenceTypes & static_cast<uint8>(EObjectReferenceType::SoftClass))
 	{
+		FText FormatString = LOCTEXT("ClassAssetTooltip", "Path to a class inheriting from type \'{TypeName}\' which may not be loaded. Can be utilized to asynchronously load the class.");
+		if (!bIsNative)
+		{
+			FormatString = LOCTEXT("ClassAssetBPTooltip", "Path to a class inheriting from type \'{TypeName}\' which may not be loaded. The blueprint type itself will always be loaded, which can be expensive.");
+		}
+
 		PinType.PinCategory = UEdGraphSchema_K2::PC_SoftClass;
-		TSharedRef<SWidget> Widget = CreateObjectReferenceWidget(InItem, PinType, IconBrush, FText::Format(LOCTEXT("ClassAssetTooltip", "Path to a class object of type \'{TypeName}\' which may be in an unloaded state. Can be utilized to asynchronously load the class."), Args));
+		TSharedRef<SWidget> Widget = CreateObjectReferenceWidget(InItem, PinType, IconBrush, FText::Format(FormatString, Args));
 		FObjectReferenceListItem ObjectReferenceType = MakeShareable(new FObjectReferenceType(InItem, Widget, PinType.PinCategory));
 		AllowedObjectReferenceTypes.Add(ObjectReferenceType);
 	}
