Use right-click to toggle value/array in Partial pin type selector

Right-clicking to toggle value/array previously only worked in the Compact type but not in the Full type which has a separate container type button.
The Partial selector type is new to UE5 and shows the type icon and name, but no container type button, so enable the right-click behavior for Partial type as well.
Practically, this change only affects the My Blueprint tab, which is the only user of ESelectorType::Partial.

diff --git a/Engine/Source/Editor/KismetWidgets/Private/SPinTypeSelector.cpp b/Engine/Source/Editor/KismetWidgets/Private/SPinTypeSelector.cpp
--- a/Engine/Source/Editor/KismetWidgets/Private/SPinTypeSelector.cpp
+++ b/Engine/Source/Editor/KismetWidgets/Private/SPinTypeSelector.cpp
@@ -195,6 +195,18 @@ static bool ContainerRequiresGetTypeHash(EPinContainerType InType)
 	return InType == EPinContainerType::Set || InType == EPinContainerType::Map;
 }
 
+static bool IsRightClickArrayStateToggleEnabled(SPinTypeSelector::ESelectorType SelectorType)
+{
+	// The Compact and Partial types both allow editing but omit the container-type
+	// (e.g. value, array) control. For both of these types, enable right-clicking
+	// on the type combo button to toggle between value/array.
+	// The None type isn't editable, so it doesn't need this behavior.
+	// The Full type shows a separate combo button for choosing the container type,
+	// so it doesn't need this right-click behavior either.
+	return SelectorType == SPinTypeSelector::ESelectorType::Compact ||
+		SelectorType == SPinTypeSelector::ESelectorType::Partial;
+}
+
 TSharedRef<SWidget> SPinTypeSelector::ConstructPinTypeImage(const FSlateBrush* PrimaryIcon, const FSlateColor& PrimaryColor, const FSlateBrush* SecondaryIcon, const FSlateColor& SecondaryColor, TSharedPtr<SToolTip> InToolTip)
 {
 	return
@@ -1315,7 +1327,7 @@ FText SPinTypeSelector::GetToolTipForComboBoxType() const
 	FText EditText;
 	if(IsEnabled())
 	{
-		if (SelectorType == ESelectorType::Compact)
+		if (IsRightClickArrayStateToggleEnabled(SelectorType))
 		{
 			EditText = LOCTEXT("CompactPinTypeSelector", "Left click to select the variable's pin type. Right click to toggle the type as an array.\n");
 		}
@@ -1412,7 +1424,7 @@ FText SPinTypeSelector::GetPinTypeItemCountText() const
 
 FReply SPinTypeSelector::OnMouseButtonDown( const FGeometry& MyGeometry, const FPointerEvent& MouseEvent )
 {
-	if (SelectorType == ESelectorType::Compact && MouseEvent.GetEffectingButton() == EKeys::RightMouseButton)
+	if (IsRightClickArrayStateToggleEnabled(SelectorType) && MouseEvent.GetEffectingButton() == EKeys::RightMouseButton)
 	{
 		bIsRightMousePressed = true;
 		return FReply::Handled();
@@ -1423,7 +1435,7 @@ FReply SPinTypeSelector::OnMouseButtonDown( const FGeometry& MyGeometry, const F
 
 FReply SPinTypeSelector::OnMouseButtonUp( const FGeometry& MyGeometry, const FPointerEvent& MouseEvent )
 {
-	if (SelectorType == ESelectorType::Compact && MouseEvent.GetEffectingButton() == EKeys::RightMouseButton)
+	if (IsRightClickArrayStateToggleEnabled(SelectorType) && MouseEvent.GetEffectingButton() == EKeys::RightMouseButton)
 	{
 		if (bIsRightMousePressed)
 		{
