Add a PinSelectorFilter to allow users to customize bluerpint pin filtering

#p4v-preflight-copy 17290371

diff --git a/Engine/Source/Editor/KismetWidgets/Private/SPinTypeSelector.cpp b/Engine/Source/Editor/KismetWidgets/Private/SPinTypeSelector.cpp
--- a/Engine/Source/Editor/KismetWidgets/Private/SPinTypeSelector.cpp
+++ b/Engine/Source/Editor/KismetWidgets/Private/SPinTypeSelector.cpp
@@ -20,6 +20,7 @@
 #include "Widgets/Layout/SWidgetSwitcher.h"
 #include "Framework/MultiBox/MultiBoxBuilder.h"
 #include "BlueprintEditorSettings.h"
+#include "PinTypeSelectorFilter.h"
 
 #define LOCTEXT_NAMESPACE "PinTypeSelector"
 
@@ -317,6 +318,11 @@ void SPinTypeSelector::Construct(const FArguments& InArgs, FGetPinTypeTree GetPi
 	{
 		CustomFilter = MakeShared<FPinTypeSelectorCustomFilterProxy>(InArgs._CustomFilter.ToSharedRef(), FSimpleDelegate::CreateSP(this, &SPinTypeSelector::OnCustomFilterChanged));
 	}
+	else if (UClass* PinTypeSelectorFilterClass = GetDefault<UPinTypeSelectorFilter>()->FilterClass.LoadSynchronous() )
+	{
+		TSharedPtr<IPinTypeSelectorFilter> SelectorFilter = GetDefault<UPinTypeSelectorFilter>(PinTypeSelectorFilterClass)->GetPinTypeSelectorFilter();
+		CustomFilter = MakeShared<FPinTypeSelectorCustomFilterProxy>(SelectorFilter.ToSharedRef(), FSimpleDelegate::CreateSP(this, &SPinTypeSelector::OnCustomFilterChanged));
+	}
 
 	bIsRightMousePressed = false;
 
diff --git a/Engine/Source/Editor/KismetWidgets/Public/PinTypeSelectorFilter.h b/Engine/Source/Editor/KismetWidgets/Public/PinTypeSelectorFilter.h
new file mode 100644
--- /dev/null
+++ b/Engine/Source/Editor/KismetWidgets/Public/PinTypeSelectorFilter.h
@@ -0,0 +1,19 @@
+// Copyright Epic Games, Inc. All Rights Reserved.
+
+#pragma once
+
+#include "CoreMinimal.h"
+#include "SPinTypeSelector.h"
+
+#include "PinTypeSelectorFilter.generated.h"
+
+UCLASS(abstract, transient, MinimalAPI, config=Editor)
+class UPinTypeSelectorFilter : public UObject
+{
+	GENERATED_BODY()
+public:
+	UPROPERTY(config)
+	TSoftClassPtr<UPinTypeSelectorFilter> FilterClass = nullptr;
+
+	virtual TSharedPtr<IPinTypeSelectorFilter> GetPinTypeSelectorFilter() const { return TSharedPtr<IPinTypeSelectorFilter>(); }
+};
\ No newline at end of file
