Fix crash in create blueprint from actor dialog.
Merged from 17457276

diff --git a/Engine/Source/Editor/ContentBrowserData/Private/ContentBrowserItemPath.cpp b/Engine/Source/Editor/ContentBrowserData/Private/ContentBrowserItemPath.cpp
new file mode 100644
--- /dev/null
+++ b/Engine/Source/Editor/ContentBrowserData/Private/ContentBrowserItemPath.cpp
@@ -0,0 +1,57 @@
+// Copyright Epic Games, Inc. All Rights Reserved.
+
+#include "ContentBrowserItemPath.h"
+#include "IContentBrowserDataModule.h"
+
+FName FContentBrowserItemPath::GetVirtualPathName() const
+{
+	return VirtualPath;
+}
+
+FName FContentBrowserItemPath::GetInternalPathName() const
+{
+	return InternalPath;
+}
+
+FString FContentBrowserItemPath::GetVirtualPathString() const
+{
+	return VirtualPath.ToString();
+}
+
+FString FContentBrowserItemPath::GetInternalPathString() const
+{
+	return InternalPath.ToString();
+}
+
+bool FContentBrowserItemPath::HasInternalPath() const
+{
+	return !InternalPath.IsNone();
+}
+
+void FContentBrowserItemPath::SetPathFromString(const FStringView InPath, const EContentBrowserPathType InPathType)
+{
+	SetPathFromName(FName(InPath), InPathType);
+}
+
+void FContentBrowserItemPath::SetPathFromName(const FName InPath, const EContentBrowserPathType InPathType)
+{
+	if (InPathType == EContentBrowserPathType::Virtual)
+	{
+		VirtualPath = InPath;
+		const EContentBrowserPathType AssetPathType = IContentBrowserDataModule::Get().GetSubsystem()->TryConvertVirtualPath(InPath, InternalPath);
+		if (AssetPathType != EContentBrowserPathType::Internal)
+		{
+			InternalPath = NAME_None;
+		}
+	}
+	else if (InPathType == EContentBrowserPathType::Internal)
+	{
+		InternalPath = InPath;
+		IContentBrowserDataModule::Get().GetSubsystem()->ConvertInternalPathToVirtual(InPath, VirtualPath);
+	}
+	else
+	{
+		InternalPath = NAME_None;
+		VirtualPath = NAME_None;
+	}
+}
diff --git a/Engine/Source/Editor/ContentBrowserData/Public/ContentBrowserDataSubsystem.h b/Engine/Source/Editor/ContentBrowserData/Public/ContentBrowserDataSubsystem.h
--- a/Engine/Source/Editor/ContentBrowserData/Public/ContentBrowserDataSubsystem.h
+++ b/Engine/Source/Editor/ContentBrowserData/Public/ContentBrowserDataSubsystem.h
@@ -15,8 +15,11 @@
 UENUM(BlueprintType)
 enum class EContentBrowserPathType : uint8
 {
+	/** No path type set */
 	None,
+	/** Internal path compatible with asset registry and engine calls (eg,. "/PluginA/MyFile") */
 	Internal,
+	/** Virtual path for enumerating Content Browser data (eg, "/All/Plugins/PluginA/MyFile") */
 	Virtual
 };
 
diff --git a/Engine/Source/Editor/ContentBrowserData/Public/ContentBrowserItemPath.h b/Engine/Source/Editor/ContentBrowserData/Public/ContentBrowserItemPath.h
new file mode 100644
--- /dev/null
+++ b/Engine/Source/Editor/ContentBrowserData/Public/ContentBrowserItemPath.h
@@ -0,0 +1,54 @@
+// Copyright Epic Games, Inc. All Rights Reserved.
+
+#pragma once
+
+#include "CoreMinimal.h"
+#include "ContentBrowserDataSubsystem.h"
+
+/**
+ * Hold multiple versions of a path as FNames
+ * 
+ * Path conversion each time Set is called
+ */
+class CONTENTBROWSERDATA_API FContentBrowserItemPath
+{
+public:
+
+	/**
+	 * Set the path being stored
+	 */
+	void SetPathFromString(const FStringView InPath, const EContentBrowserPathType InPathType);
+	void SetPathFromName(const FName InPath, const EContentBrowserPathType InPathType);
+
+	/**
+	 * Returns virtual path as FName (eg, "/All/Plugins/PluginA/MyFile").
+	 */
+	FName GetVirtualPathName() const;
+
+	/**
+	 * Returns internal path if there is one (eg,. "/PluginA/MyFile").
+	 */
+	FName GetInternalPathName() const;
+
+	/**
+	 * Returns virtual path as newly allocated FString (eg, "/All/Plugins/PluginA/MyFile").
+	 */
+	FString GetVirtualPathString() const;
+
+	/**
+	 * Returns internal path as newly allocated FString if there is one or empty FString (eg,. "/PluginA/MyFile").
+	 */
+	FString GetInternalPathString() const;
+
+	/**
+	 * Returns true if there is an internal path
+	 */
+	bool HasInternalPath() const;
+
+private:
+	/** Path as virtual (eg, "/All/Plugins/PluginA/MyFile") */
+	FName VirtualPath;
+
+	/** Path as internal (eg,. "/PluginA/MyFile") */
+	FName InternalPath;
+};
diff --git a/Engine/Source/Editor/KismetWidgets/Private/CreateBlueprintFromActorDialog.cpp b/Engine/Source/Editor/KismetWidgets/Private/CreateBlueprintFromActorDialog.cpp
--- a/Engine/Source/Editor/KismetWidgets/Private/CreateBlueprintFromActorDialog.cpp
+++ b/Engine/Source/Editor/KismetWidgets/Private/CreateBlueprintFromActorDialog.cpp
@@ -30,6 +30,7 @@
 #include "ContentBrowserModule.h"
 #include "IContentBrowserDataModule.h"
 #include "ContentBrowserDataSubsystem.h"
+#include "ContentBrowserItemPath.h"
 #include "IContentBrowserSingleton.h"
 #include "PackageTools.h"
 #include "DetailLayoutBuilder.h"
@@ -115,7 +116,7 @@ public:
 	TWeakObjectPtr<AActor> ActorOverride;
 
 	/** The path the asset should be created at */
-	FString AssetPath;
+	FContentBrowserItemPath AssetPath;
 
 	/** The the name for the new asset */
 	FString AssetName;
@@ -305,7 +306,7 @@ void SSCreateBlueprintPicker::Construct(const FArguments& InArgs)
 	ClassViewer = StaticCastSharedRef<SClassViewer>(ClassViewerModule.CreateClassViewer(ClassViewerOptions, FOnClassPicked::CreateSP(this, &SSCreateBlueprintPicker::OnClassPicked)));
 
 	FString PackageName;
-	AssetPath = ContentBrowserModule.Get().GetCurrentPath(EContentBrowserPathType::Virtual);
+	AssetPath.SetPathFromString(ContentBrowserModule.Get().GetCurrentPath(EContentBrowserPathType::Virtual), EContentBrowserPathType::Virtual);
 
 	ECreateBlueprintFromActorMode ValidCreateMethods = FCreateBlueprintFromActorDialog::GetValidCreationMethods();
 
@@ -326,8 +327,7 @@ void SSCreateBlueprintPicker::Construct(const FArguments& InArgs)
 
 	AssetName = UPackageTools::SanitizePackageName(AssetName + TEXT("Blueprint"));
 
-	FString BasePath = AssetPath / AssetName;
-	AssetToolsModule.Get().CreateUniqueAssetName(BasePath, TEXT(""), PackageName, AssetName);
+	AssetToolsModule.Get().CreateUniqueAssetName(AssetPath.GetInternalPathString() / AssetName, TEXT(""), PackageName, AssetName);
 
 	TSharedPtr<SGridPanel> CreationMethodSection;
 
@@ -434,7 +434,7 @@ void SSCreateBlueprintPicker::Construct(const FArguments& InArgs)
 						.FillWidth(1.f)
 						[
 							SNew(SEditableTextBox)
-							.Text(TAttribute<FText>::Create(TAttribute<FText>::FGetter::CreateLambda([this]() { return FText::FromString(AssetPath); })))
+							.Text(TAttribute<FText>::Create(TAttribute<FText>::FGetter::CreateLambda([this]() { return FText::FromName(AssetPath.GetVirtualPathName()); })))
 							.IsReadOnly(true)
 						]
 						+SHorizontalBox::Slot()
@@ -546,13 +546,13 @@ class SSCreateBlueprintPathPicker : public SCompoundWidget
 	{}
 
 	SLATE_ARGUMENT(TSharedPtr<SWindow>, ParentWindow)
-		SLATE_ARGUMENT(FString, AssetPath)
+		SLATE_ARGUMENT(FContentBrowserItemPath, AssetPath)
 		SLATE_END_ARGS()
 
 	void Construct(const FArguments& InArgs);
 
 	/** Callback when the selected asset path has changed. */
-	void OnSelectAssetPath(const FString& Path) { AssetPath = Path; }
+	void OnSelectAssetPath(const FString& InVirtualPath) { AssetPath.SetPathFromString(InVirtualPath, EContentBrowserPathType::Virtual); }
 
 	/** Callback when the "ok" button is clicked. */
 	FReply OnClickOk();
@@ -563,7 +563,7 @@ class SSCreateBlueprintPathPicker : public SCompoundWidget
 	/** A pointer to the window that is asking the user to select a parent class */
 	TWeakPtr<SWindow> WeakParentWindow;
 
-	FString AssetPath;
+	FContentBrowserItemPath AssetPath;
 
 	bool bPressedOk;
 };
@@ -578,7 +578,7 @@ void SSCreateBlueprintPathPicker::Construct(const FArguments& InArgs)
 	FContentBrowserModule& ContentBrowserModule = FModuleManager::LoadModuleChecked<FContentBrowserModule>("ContentBrowser");
 
 	FPathPickerConfig PathPickerConfig;
-	PathPickerConfig.DefaultPath = AssetPath;
+	PathPickerConfig.DefaultPath = AssetPath.GetVirtualPathString();
 	PathPickerConfig.OnPathSelected = FOnPathSelected::CreateRaw(this, &SSCreateBlueprintPathPicker::OnSelectAssetPath);
 	PathPickerConfig.bAllowReadOnlyFolders = false;
 	PathPickerConfig.bOnPathSelectedPassesVirtualPaths = true;
@@ -665,7 +665,7 @@ FReply SSCreateBlueprintPicker::OnPathPickerSummoned()
 
 	if (PathPickerDialog->bPressedOk)
 	{
-		AssetPath = PathPickerDialog->AssetPath;
+		AssetPath.SetPathFromString(PathPickerDialog->AssetPath.GetVirtualPathString(), EContentBrowserPathType::Virtual);
 		UpdateFilenameStatus();
 	}
 
@@ -699,7 +699,7 @@ void SSCreateBlueprintPicker::UpdateFilenameStatus()
 	{
 		TArray<FAssetData> AssetData;
 		FAssetRegistryModule& AssetRegistryModule = FModuleManager::LoadModuleChecked<FAssetRegistryModule>("AssetRegistry");
-		AssetRegistryModule.Get().GetAssetsByPath(FName(*AssetPath), AssetData);
+		AssetRegistryModule.Get().GetAssetsByPath(AssetPath.GetInternalPathName(), AssetData);
 
 		// Check to see if the name conflicts
 		for (const FAssetData& Data : AssetData)
@@ -780,12 +780,9 @@ void FCreateBlueprintFromActorDialog::OpenDialog(ECreateBlueprintFromActorMode C
 
 	if (ClassPickerDialog->bPressedOk)
 	{
-		FString VirtualAssetPath = ClassPickerDialog->AssetPath;
-		FName InternalAssetPath;
-		EContentBrowserPathType AssetPathType = IContentBrowserDataModule::Get().GetSubsystem()->TryConvertVirtualPath(VirtualAssetPath, InternalAssetPath);
-		if(ensureMsgf(AssetPathType == EContentBrowserPathType::Internal, TEXT("Could not convert virtual path %s to internal path."), *VirtualAssetPath))
+		if(ensureMsgf(ClassPickerDialog->AssetPath.HasInternalPath(), TEXT("Could not convert virtual path %s to internal path."), *ClassPickerDialog->AssetPath.GetVirtualPathString()))
 		{
-			FString NewAssetName = InternalAssetPath.ToString() / ClassPickerDialog->AssetName;
+			FString NewAssetName = ClassPickerDialog->AssetPath.GetInternalPathString() / ClassPickerDialog->AssetName;
 			OnCreateBlueprint(NewAssetName, ClassPickerDialog->ChosenClass, ClassPickerDialog->CreateMode, ActorOverride.Get(), bInReplaceActors);
 		}
 	}
