Provide default path for CreateBlueprintFromActorDialog if the Content Drawer hasn't been opened yet

diff --git a/Engine/Source/Editor/KismetWidgets/Private/CreateBlueprintFromActorDialog.cpp b/Engine/Source/Editor/KismetWidgets/Private/CreateBlueprintFromActorDialog.cpp
--- a/Engine/Source/Editor/KismetWidgets/Private/CreateBlueprintFromActorDialog.cpp
+++ b/Engine/Source/Editor/KismetWidgets/Private/CreateBlueprintFromActorDialog.cpp
@@ -13,6 +13,7 @@
 #include "Widgets/Layout/SExpandableArea.h"
 #include "Widgets/Layout/SGridPanel.h"
 #include "Widgets/Layout/SUniformGridPanel.h"
+#include "Widgets/Notifications/SNotificationList.h"
 #include "Framework/Application/SlateApplication.h"
 #include "Framework/Docking/TabManager.h"
 #include "Editor/EditorEngine.h"
@@ -306,7 +307,19 @@ void SSCreateBlueprintPicker::Construct(const FArguments& InArgs)
 	ClassViewer = StaticCastSharedRef<SClassViewer>(ClassViewerModule.CreateClassViewer(ClassViewerOptions, FOnClassPicked::CreateSP(this, &SSCreateBlueprintPicker::OnClassPicked)));
 
 	FString PackageName;
-	AssetPath.SetPathFromString(ContentBrowserModule.Get().GetCurrentPath(EContentBrowserPathType::Virtual), EContentBrowserPathType::Virtual);
+
+	FString PathString = ContentBrowserModule.Get().GetCurrentPath(EContentBrowserPathType::Virtual);
+	if (!PathString.IsEmpty())
+	{
+		AssetPath.SetPathFromString(PathString, EContentBrowserPathType::Virtual);
+	}
+	else
+	{
+		static const FString DefaultInvariantPath = TEXT("/Game");
+		FName DefaultVirtualPath;
+		IContentBrowserDataModule::Get().GetSubsystem()->ConvertInternalPathToVirtual(FStringView(DefaultInvariantPath), DefaultVirtualPath);
+		AssetPath.SetPathFromName(DefaultVirtualPath, EContentBrowserPathType::Virtual);
+	}
 
 	ECreateBlueprintFromActorMode ValidCreateMethods = FCreateBlueprintFromActorDialog::GetValidCreationMethods();
 
@@ -780,11 +793,17 @@ void FCreateBlueprintFromActorDialog::OpenDialog(ECreateBlueprintFromActorMode C
 
 	if (ClassPickerDialog->bPressedOk)
 	{
-		if(ensureMsgf(ClassPickerDialog->AssetPath.HasInternalPath(), TEXT("Could not convert virtual path %s to internal path."), *ClassPickerDialog->AssetPath.GetVirtualPathString()))
+		if (ClassPickerDialog->AssetPath.HasInternalPath())
 		{
 			FString NewAssetName = ClassPickerDialog->AssetPath.GetInternalPathString() / ClassPickerDialog->AssetName;
 			OnCreateBlueprint(NewAssetName, ClassPickerDialog->ChosenClass, ClassPickerDialog->CreateMode, ActorOverride.Get(), bInReplaceActors);
 		}
+		else
+		{
+			FNotificationInfo ErrorNotificationInfo(FText::FormatOrdered(LOCTEXT("PathError", "Could not convert virtual path '{0}' to internal path."), FText::FromString(*ClassPickerDialog->AssetPath.GetVirtualPathString())));
+			TSharedPtr<SNotificationItem> NotificationItem = FSlateNotificationManager::Get().AddNotification(ErrorNotificationInfo);
+			NotificationItem->SetCompletionState(SNotificationItem::CS_Fail);
+		}
 	}
 }
 
