Tweaks to appearance and layout of SPinTypeSelector

- Tweak padding of read-only representation to prevent pin selector label from shifting when switching between read-only and editable (e.g. hovering in My Blueprint)
- In Full selector type, adjust horizontal box slot sizing to prevent combo boxes from shifting when changing the map value type (e.g. in Blueprint Details)
- In Partial selector type, show combined type (e.g. "Array of Strings" or "Map of Strings to Integers") for containers and display layered primary+secondary type icon for maps (e.g. in My Blueprint)

diff --git a/Engine/Source/Editor/KismetWidgets/Private/SPinTypeSelector.cpp b/Engine/Source/Editor/KismetWidgets/Private/SPinTypeSelector.cpp
--- a/Engine/Source/Editor/KismetWidgets/Private/SPinTypeSelector.cpp
+++ b/Engine/Source/Editor/KismetWidgets/Private/SPinTypeSelector.cpp
@@ -326,6 +326,29 @@ void SPinTypeSelector::Construct(const FArguments& InArgs, FGetPinTypeTree GetPi
 
 	bIsRightMousePressed = false;
 
+	// Depending on whether this is a full selector or not, we generate a different primary type image widget
+	TSharedPtr<SWidget> PrimaryTypeImage;
+	if (SelectorType == ESelectorType::Full)
+	{
+		// Full selector displays container and secondary type separately, so the main combo button should just have the primary icon
+		PrimaryTypeImage =
+			SNew(SImage)
+			.Image(this, &SPinTypeSelector::GetTypeIconImage)
+			.ColorAndOpacity(this, &SPinTypeSelector::GetTypeIconColor);
+	}
+	else
+	{
+		// Partial/compact selectors do not display container or secondary type separately, so we need to jam it all in the one image
+		PrimaryTypeImage =
+			SNew(
+				SLayeredImage,
+				TAttribute<const FSlateBrush*>(this, &SPinTypeSelector::GetSecondaryTypeIconImage),
+				TAttribute<FSlateColor>(this, &SPinTypeSelector::GetSecondaryTypeIconColor)
+			)
+			.Image(this, &SPinTypeSelector::GetTypeIconImage)
+			.ColorAndOpacity(this, &SPinTypeSelector::GetTypeIconColor);
+	}
+
 	// Depending on if this is a compact selector or not, we generate a different compound widget
 	TSharedPtr<SWidget> Widget;
 
@@ -340,25 +363,13 @@ void SPinTypeSelector::Construct(const FArguments& InArgs, FGetPinTypeTree GetPi
 			.ButtonStyle(FEditorStyle::Get(),  "BlueprintEditor.CompactPinTypeSelector")
 			.ButtonContent()
 			[
-				SNew(
-					SLayeredImage,
-					TAttribute<const FSlateBrush*>(this, &SPinTypeSelector::GetSecondaryTypeIconImage),
-					TAttribute<FSlateColor>(this, &SPinTypeSelector::GetSecondaryTypeIconColor)
-				)
-				.Image(this, &SPinTypeSelector::GetTypeIconImage)
-				.ColorAndOpacity(this, &SPinTypeSelector::GetTypeIconColor)
+				PrimaryTypeImage.ToSharedRef()
 			];
 	}
 	else if (SelectorType == ESelectorType::None)
 	{
-		Widget = SNew(
-					SLayeredImage,
-					TAttribute<const FSlateBrush*>(this, &SPinTypeSelector::GetSecondaryTypeIconImage),
-					TAttribute<FSlateColor>(this, &SPinTypeSelector::GetSecondaryTypeIconColor)
-				)
-				.ToolTipText(this, &SPinTypeSelector::GetToolTipForComboBoxType)
-				.Image(this, &SPinTypeSelector::GetTypeIconImage)
-				.ColorAndOpacity(this, &SPinTypeSelector::GetTypeIconColor);
+		Widget = PrimaryTypeImage.ToSharedRef();
+		Widget->SetToolTipText(TAttribute<FText>::CreateSP(this, &SPinTypeSelector::GetToolTipForComboBoxType));
 	}
 	else if (SelectorType == ESelectorType::Full || SelectorType == ESelectorType::Partial)
 	{
@@ -386,11 +397,13 @@ void SPinTypeSelector::Construct(const FArguments& InArgs, FGetPinTypeTree GetPi
 		TSharedRef<SHorizontalBox> HBox = SNew(SHorizontalBox).Clipping(EWidgetClipping::ClipToBoundsAlways);
 		Widget = HBox;
 
+		const float FullComboButtonWidth = 125.f;
+
 		HBox->AddSlot()
 		.HAlign(HAlign_Left)
 		[
 			SNew(SBox)
-			.WidthOverride(SelectorType == ESelectorType::Full ? 125.f : FOptionalSize())
+			.WidthOverride(SelectorType == ESelectorType::Full ? FullComboButtonWidth : FOptionalSize())
 			[
 				SAssignNew(TypeComboButton, SComboButton)
 				.ComboButtonStyle(FAppStyle::Get(), SelectorType == ESelectorType::Full ? "ComboButton" : "ComboButton")
@@ -408,9 +421,7 @@ void SPinTypeSelector::Construct(const FArguments& InArgs, FGetPinTypeTree GetPi
 					.Padding(0.0f, 0.0f, 2.0f, 0.0f)
 					.AutoWidth()
 					[
-						SNew(SImage)
-						.Image(this, &SPinTypeSelector::GetTypeIconImage)
-						.ColorAndOpacity(this, &SPinTypeSelector::GetTypeIconColor)
+						PrimaryTypeImage.ToSharedRef()
 					]
 					+ SHorizontalBox::Slot()
 					.Padding(2.0f, 0.0f, 0.0f, 0.0f)
@@ -441,6 +452,7 @@ void SPinTypeSelector::Construct(const FArguments& InArgs, FGetPinTypeTree GetPi
 			HBox->AddSlot()
 			[
 				SNew(SBox)
+				.WidthOverride(FullComboButtonWidth)
 				.Visibility_Lambda([this]() {return this->TargetPinType.Get().IsMap() == true ? EVisibility::Visible : EVisibility::Collapsed; })
 				[
 					SAssignNew( SecondaryTypeComboButton, SComboButton )
@@ -481,22 +493,23 @@ void SPinTypeSelector::Construct(const FArguments& InArgs, FGetPinTypeTree GetPi
 	[
 		SNew(SWidgetSwitcher)
 		.WidgetIndex_Lambda([this](){return ReadOnly.Get() ? 1 : 0; })
-		+ SWidgetSwitcher::Slot()
+		+ SWidgetSwitcher::Slot() // editable version
 		.Padding(SelectorType == ESelectorType::Partial ? FMargin(-6.0f, 0.0f,0.0f,0.0f) : FMargin(0))
 		[
 			Widget.ToSharedRef()
 		]
-		+ SWidgetSwitcher::Slot()
+		+ SWidgetSwitcher::Slot() // read-only version
 		[
 			SNew(SHorizontalBox)
 			.Clipping(EWidgetClipping::OnDemand)
 			+ SHorizontalBox::Slot()
 			.VAlign(VAlign_Center)
 			.HAlign(HAlign_Left)
-			.Padding(FMargin(2.0f, 3.0f, 0.0f, 3.0f))
+			.Padding(FMargin(2.0f, 3.0f, 2.0f, 3.0f))
 			.AutoWidth()
 			[
-				SNew(SImage)
+				// Read-only version does not display container or secondary type separately, so we need to jam it all in the one image
+				SNew(SLayeredImage, TAttribute<const FSlateBrush*>(this, &SPinTypeSelector::GetSecondaryTypeIconImage), TAttribute<FSlateColor>(this, &SPinTypeSelector::GetSecondaryTypeIconColor))
 				.Image(this, &SPinTypeSelector::GetTypeIconImage)
 				.ColorAndOpacity(this, &SPinTypeSelector::GetTypeIconColor)
 			]
@@ -554,6 +567,27 @@ FText SPinTypeSelector::GetSecondaryTypeDescription() const
 	}
 }
 
+FText SPinTypeSelector::GetCombinedTypeDescription() const
+{
+	FFormatNamedArguments Args;
+
+	EPinContainerType ContainerType = TargetPinType.Get().ContainerType;
+	switch (TargetPinType.Get().ContainerType)
+	{
+	case EPinContainerType::Map:
+		Args.Add(TEXT("KeyTitle"), GetTypeDescription());
+		Args.Add(TEXT("ValueTitle"), GetSecondaryTypeDescription());
+		return FText::Format(NSLOCTEXT("KismetSchema", "MapAsText", "Map of {KeyTitle}s to {ValueTitle}s"), Args);
+	case EPinContainerType::Set:
+		Args.Add(TEXT("PropertyTitle"), GetTypeDescription());
+		return FText::Format(NSLOCTEXT("SetAsText", "MapAsText", "Set of {PropertyTitle}s"), Args);
+	case EPinContainerType::Array:
+		Args.Add(TEXT("PropertyTitle"), GetTypeDescription());
+		return FText::Format(NSLOCTEXT("ArrayAsText", "MapAsText", "Array of {PropertyTitle}s"), Args);
+	default:
+		return GetTypeDescription();
+	}
+}
 
 const FSlateBrush* SPinTypeSelector::GetTypeIconImage() const
 {
@@ -1342,7 +1376,9 @@ FText SPinTypeSelector::GetToolTipForComboBoxType() const
 		EditText = LOCTEXT("PinTypeSelector_Disabled", "Cannot edit variable type when they are inherited from parent.\n");
 	}
 
-	return FText::Format(LOCTEXT("PrimaryTypeTwoLines", "{0}Current Type: {1}"), EditText, GetTypeDescription());
+	// With the full selector type, just display the primary type in the tooltip because the secondary type has its own combo box.
+	// With the partial selector type, we need to jam everything into the tooltip for the single combo box.
+	return FText::Format(LOCTEXT("PrimaryTypeTwoLines", "{0}Current Type: {1}"), EditText, SelectorType == ESelectorType::Full ? GetTypeDescription() : GetCombinedTypeDescription());
 }
 
 FText SPinTypeSelector::GetToolTipForComboBoxSecondaryType() const
diff --git a/Engine/Source/Editor/KismetWidgets/Public/SPinTypeSelector.h b/Engine/Source/Editor/KismetWidgets/Public/SPinTypeSelector.h
--- a/Engine/Source/Editor/KismetWidgets/Public/SPinTypeSelector.h
+++ b/Engine/Source/Editor/KismetWidgets/Public/SPinTypeSelector.h
@@ -112,6 +112,9 @@ protected:
 	/** Gets the secondary type description. E.g. the value type for TMaps */
 	FText GetSecondaryTypeDescription() const;
 
+	/** Gets a combined description of the primary, container, and secondary types. E.g. "Map of Strings to Floats" */
+	FText GetCombinedTypeDescription() const;
+
 	TSharedPtr<SComboButton>		TypeComboButton;
 	TSharedPtr<SComboButton>		SecondaryTypeComboButton;
 	TSharedPtr<SSearchBox>			FilterTextBox;
