First pass at interactive Blueprint graph pin value inspection tooltips while debugging in PIE.

Currently disabled by default; enable via setting "BP.EnablePinValueInspectionDuringPIE 1" at the console, then PIE and break execution with an instance of the open Blueprint graph selected for debugging.

diff --git a/Engine/Source/Editor/GraphEditor/Private/SGraphPin.cpp b/Engine/Source/Editor/GraphEditor/Private/SGraphPin.cpp
--- a/Engine/Source/Editor/GraphEditor/Private/SGraphPin.cpp
+++ b/Engine/Source/Editor/GraphEditor/Private/SGraphPin.cpp
@@ -3,6 +3,7 @@
 
 #include "SGraphPin.h"
 #include "Widgets/SBoxPanel.h"
+#include "Widgets/SToolTip.h"
 #include "Widgets/Layout/SWrapBox.h"
 #include "Widgets/Images/SImage.h"
 #include "Widgets/Layout/SBox.h"
@@ -19,6 +20,7 @@
 #include "ScopedTransaction.h"
 #include "SLevelOfDetailBranchNode.h"
 #include "SPinTypeSelector.h"
+#include "SPinValueInspector.h"
 #include "Animation/AnimNodeBase.h"
 
 /////////////////////////////////////////////////////
@@ -310,8 +312,12 @@ void SGraphPin::Construct(const FArguments& InArgs, UEdGraphPin* InPin)
 		]
 	);
 
-	TAttribute<FText> ToolTipAttribute = TAttribute<FText>::Create(TAttribute<FText>::FGetter::CreateSP(this, &SGraphPin::GetTooltipText));
-	SetToolTipText(ToolTipAttribute);
+	TSharedPtr<IToolTip> TooltipWidget = SNew(SToolTip)
+		.Text(this, &SGraphPin::GetTooltipText)
+		.IsInteractive(this, &SGraphPin::IsTooltipInteractive)
+		.OnSetInteractiveWindowLocation(this, &SGraphPin::OnSetInteractiveTooltipLocation);
+
+	SetToolTip(TooltipWidget);
 }
 
 TSharedRef<SWidget>	SGraphPin::GetDefaultValueWidget()
@@ -1192,6 +1198,39 @@ void SGraphPin::SetOnlyShowDefaultValue(bool bNewOnlyShowDefaultValue)
 	bOnlyShowDefaultValue = bNewOnlyShowDefaultValue;
 }
 
+TSharedPtr<IToolTip> SGraphPin::GetToolTip()
+{
+	TSharedPtr<IToolTip> CurrentTooltip = SBorder::GetToolTip();
+	if (CurrentTooltip.IsValid())
+	{
+		const UEdGraphPin* GraphPin = GetPinObj();
+		if (GraphPin && FKismetDebugUtilities::CanInspectPinValue(GraphPin) && !ValueInspectorWidget.IsValid())
+		{
+			ValueInspectorWidget = SNew(SPinValueInspector, GraphPin);
+
+			CurrentTooltip->SetContentWidget(ValueInspectorWidget.ToSharedRef());
+		}
+	}
+
+	return CurrentTooltip;
+}
+
+void SGraphPin::OnToolTipClosing()
+{
+	if (ValueInspectorWidget.IsValid())
+	{
+		TSharedPtr<IToolTip> CurrentTooltip = SBorder::GetToolTip();
+		if (CurrentTooltip.IsValid())
+		{
+			CurrentTooltip->ResetContentWidget();
+		}
+
+		ValueInspectorWidget.Reset();
+	}
+
+	SBorder::OnToolTipClosing();
+}
+
 FText SGraphPin::GetTooltipText() const
 {
 	if(!ensure(!bGraphDataInvalid))
@@ -1216,6 +1255,52 @@ FText SGraphPin::GetTooltipText() const
 	return HoverText;
 }
 
+bool SGraphPin::IsTooltipInteractive() const
+{
+	return ValueInspectorWidget.IsValid();
+}
+
+void SGraphPin::OnSetInteractiveTooltipLocation(FVector2D& InOutDesiredLocation) const
+{
+	TSharedPtr<SGraphNode> OwnerNode = OwnerNodePtr.Pin();
+	if (OwnerNode.IsValid())
+	{
+		TSharedPtr<SGraphPanel> GraphPanel = OwnerNode->GetOwnerPanel();
+		if (GraphPanel.IsValid())
+		{
+			// Reset to the pin's location in graph space.
+			InOutDesiredLocation = OwnerNode->GetPosition() + CachedNodeOffset;
+			
+			// Shift the desired location to the right edge of the pin's geometry.
+			InOutDesiredLocation.X += GetTickSpaceGeometry().Size.X;
+
+			// Align to the first entry in the inspector's tree view.
+			if (ValueInspectorWidget.IsValid())
+			{
+				// @todo - Find a way to calculate these at runtime, e.g. based off of actual child widget geometry?
+				static const float VerticalOffsetWithSearchFilter = 41.0f;
+				static const float VerticalOffsetWithoutSearchFilter = 19.0f;
+
+				if (ValueInspectorWidget->ShouldShowSearchFilter())
+				{
+					InOutDesiredLocation.Y -= VerticalOffsetWithSearchFilter;
+				}
+				else
+				{
+					InOutDesiredLocation.Y -= VerticalOffsetWithoutSearchFilter;
+				}
+			}
+
+			// Convert our desired location from graph coordinates into panel space.
+			InOutDesiredLocation -= GraphPanel->GetViewOffset();
+			InOutDesiredLocation *= GraphPanel->GetZoomAmount();
+
+			// Finally, convert the modified location from panel space into screen space.
+			InOutDesiredLocation = GraphPanel->GetTickSpaceGeometry().LocalToAbsolute(InOutDesiredLocation);
+		}
+	}
+}
+
 bool SGraphPin::IsEditingEnabled() const
 {
 	if (OwnerNodePtr.IsValid())
diff --git a/Engine/Source/Editor/GraphEditor/Public/SGraphPin.h b/Engine/Source/Editor/GraphEditor/Public/SGraphPin.h
--- a/Engine/Source/Editor/GraphEditor/Public/SGraphPin.h
+++ b/Engine/Source/Editor/GraphEditor/Public/SGraphPin.h
@@ -20,6 +20,7 @@ class SGraphPin;
 class SHorizontalBox;
 class SImage;
 class SWrapBox;
+class SPinValueInspector;
 
 #define NAME_DefaultPinLabelStyle TEXT("Graph.Node.PinName")
 
@@ -118,6 +119,8 @@ public:
 	virtual FReply OnDragOver(const FGeometry& MyGeometry, const FDragDropEvent& DragDropEvent) override;
 	virtual FReply OnDrop(const FGeometry& MyGeometry, const FDragDropEvent& DragDropEvent) override;
 	virtual void Tick( const FGeometry& AllottedGeometry, const double InCurrentTime, const float InDeltaTime ) override;
+	virtual TSharedPtr<IToolTip> GetToolTip() override;
+	virtual void OnToolTipClosing() override;
 	// End of SWidget interface
 
 public:
@@ -251,6 +254,12 @@ protected:
 	/** @return The tooltip to display for this pin */
 	FText GetTooltipText() const;
 
+	/** @return Whether the tooltip should be considered interactive */
+	bool IsTooltipInteractive() const;
+
+	/** Delegate to modify or reset the given window location (in screen coords) for an interactive tooltip (e.g. pin value inspector) */
+	void OnSetInteractiveTooltipLocation(FVector2D& InOutDesiredLocation) const;
+
 	TOptional<EMouseCursor::Type> GetPinCursor() const;
 
 	/** Spawns a FDragConnection or similar class for the pin drag event */
@@ -278,6 +287,9 @@ protected:
 	/** Value widget for the pin, created with GetDefaultValueWidget() */
 	TSharedPtr<SWidget> ValueWidget;
 
+	/** Value inspector widget hosted inside the tooltip while debugging */
+	TSharedPtr<SPinValueInspector> ValueInspectorWidget;
+
 	/** The GraphPin that this widget represents. */
 	class UEdGraphPin* GraphPinObj;
 
diff --git a/Engine/Source/Editor/KismetWidgets/Private/SPinValueInspector.cpp b/Engine/Source/Editor/KismetWidgets/Private/SPinValueInspector.cpp
new file mode 100644
--- /dev/null
+++ b/Engine/Source/Editor/KismetWidgets/Private/SPinValueInspector.cpp
@@ -0,0 +1,331 @@
+// Copyright Epic Games, Inc. All Rights Reserved.
+
+#include "SPinValueInspector.h"
+#include "UObject/WeakFieldPtr.h"
+#include "Styling/AppStyle.h"
+#include "Widgets/SBoxPanel.h"
+#include "PropertyInfoViewStyle.h"
+#include "Widgets/Input/SSearchBox.h"
+#include "Kismet2/BlueprintEditorUtils.h"
+#include "Kismet2/KismetDebugUtilities.h"
+#include "IDetailPropertyRow.h"
+
+#define LOCTEXT_NAMESPACE "SPinValueInspector"
+
+struct FPinValueInspectorTreeViewNode
+{
+	FText DisplayName;
+	FText Description;
+	FSlateColor IconColor;
+	const FSlateBrush* IconBrush;
+	TArray<TSharedPtr<FPropertyInstanceInfo>> Children;
+	TSharedPtr<FDetailColumnSizeData> SharedColumnSizeData;
+
+	FPinValueInspectorTreeViewNode()
+		:IconBrush(FEditorStyle::GetBrush(TEXT("Kismet.AllClasses.VariableIcon")))
+	{}
+};
+
+class SPinValueInspector_ConstrainedBox : public SCompoundWidget
+{
+public:
+	SLATE_BEGIN_ARGS(SPinValueInspector_ConstrainedBox)
+		: _MinWidth()
+		, _MaxWidth()
+	{}
+		SLATE_DEFAULT_SLOT(FArguments, Content)
+		SLATE_ATTRIBUTE(TOptional<float>, MinWidth)
+		SLATE_ATTRIBUTE(TOptional<float>, MaxWidth)
+
+	SLATE_END_ARGS()
+
+	void Construct(const FArguments& InArgs)
+	{
+		MinWidth = InArgs._MinWidth;
+		MaxWidth = InArgs._MaxWidth;
+
+		ChildSlot
+		[
+			InArgs._Content.Widget
+		];
+	}
+	
+	virtual FVector2D ComputeDesiredSize(float LayoutScaleMultiplier) const override
+	{
+		const float MinWidthVal = MinWidth.Get().Get(0.0f);
+		const float MaxWidthVal = MaxWidth.Get().Get(0.0f);
+
+		if (MinWidthVal == 0.0f && MaxWidthVal == 0.0f)
+		{
+			return SCompoundWidget::ComputeDesiredSize(LayoutScaleMultiplier);
+		}
+		else
+		{
+			FVector2D ChildSize = ChildSlot.GetWidget()->GetDesiredSize();
+
+			float XVal = FMath::Max(MinWidthVal, ChildSize.X);
+			if (MaxWidthVal > MinWidthVal)
+			{
+				XVal = FMath::Min(MaxWidthVal, XVal);
+			}
+
+			return FVector2D(XVal, ChildSize.Y);
+		}
+	}
+
+private:
+	TAttribute<TOptional<float>> MinWidth;
+	TAttribute<TOptional<float>> MaxWidth;
+};
+
+class SPinValueInspector_TreeViewRow : public STableRow<FPinValueInspectorTreeViewNodePtr>
+{
+public:
+	SLATE_BEGIN_ARGS(SPinValueInspector_TreeViewRow) {}
+	SLATE_END_ARGS()
+
+	void Construct(const FArguments& InArgs, TSharedRef<STableViewBase> OwnerTableView, FPinValueInspectorTreeViewNodePtr InNode, bool bShowExpanderArrows)
+	{
+		STableRow<FPinValueInspectorTreeViewNodePtr>::ConstructInternal(STableRow::FArguments(), OwnerTableView);
+
+		if (!InNode.IsValid() || !InNode->SharedColumnSizeData.IsValid())
+		{
+			return;
+		}
+
+		TSharedRef<SWidget> RowWidget =
+			SNew(SSplitter)
+			.Style(FEditorStyle::Get(), "DetailsView.Splitter")
+			.PhysicalSplitterHandleSize(1.0f)
+			.HitDetectionSplitterHandleSize(5.0f)
+			+SSplitter::Slot()
+			.Value(InNode->SharedColumnSizeData->NameColumnWidth)
+			.OnSlotResized(InNode->SharedColumnSizeData->OnNameColumnResized)
+			[
+				SNew(SHorizontalBox)
+				.Clipping(EWidgetClipping::OnDemand)
+				+SHorizontalBox::Slot()
+				.HAlign(HAlign_Left)
+				.VAlign(VAlign_Fill)
+				.AutoWidth()
+				[
+					SNew(PropertyInfoViewStyle::SIndent, SharedThis(this))
+				]
+				+SHorizontalBox::Slot()
+				.HAlign(HAlign_Left)
+				.VAlign(VAlign_Center)
+				.Padding(8.0f, 0.0f, 0.0f, 0.0f)
+				.AutoWidth()
+				[
+					SNew(PropertyInfoViewStyle::SExpanderArrow, SharedThis(this))
+					.HasChildren(InNode->Children.Num() > 0)
+					.Visibility(bShowExpanderArrows ? EVisibility::Visible : EVisibility::Collapsed)
+				]
+				+SHorizontalBox::Slot()
+				.HAlign(HAlign_Left)
+				.VAlign(VAlign_Center)
+				.AutoWidth()
+				[
+					SNew(SImage)
+					.Image(InNode->IconBrush)
+					.ColorAndOpacity(InNode->IconColor)
+				]
+				+SHorizontalBox::Slot()
+				.HAlign(HAlign_Left)
+				.VAlign(VAlign_Center)
+				.Padding(5.0f, 0.0f, 0.0f, 0.0f)
+				[
+					SNew(STextBlock)
+					.Text(InNode->DisplayName)
+				]
+			]
+			+SSplitter::Slot()
+			.Value(InNode->SharedColumnSizeData->ValueColumnWidth)
+			.OnSlotResized(InNode->SharedColumnSizeData->OnValueColumnResized)
+			[
+				SNew(SHorizontalBox)
+				.Clipping(EWidgetClipping::OnDemand)
+				+SHorizontalBox::Slot()
+				.HAlign(HAlign_Left)
+				.VAlign(VAlign_Center)
+				[
+					SNew(SPinValueInspector_ConstrainedBox)
+					.MinWidth(125.0f)
+					.MaxWidth(400.0f)
+					[
+						SNew(STextBlock)
+						.Text(InNode->Description)
+					]
+				]
+			];
+		
+		ChildSlot
+		[
+			SNew(SBorder)
+			.Clipping(EWidgetClipping::ClipToBounds)
+			[
+				RowWidget
+			]
+		];
+	}
+};
+
+void SPinValueInspector::Construct(const FArguments& InArgs, const FEdGraphPinReference& InPinRef)
+{
+	PinRef = InPinRef;
+
+	// Locate the class property associated with the source pin and set it to the root node.
+	const UEdGraphPin* GraphPin = PinRef.Get();
+	if (GraphPin)
+	{
+		if (const UEdGraphNode* GraphNode = GraphPin->GetOwningNodeUnchecked())
+		{
+			if (UBlueprint* Blueprint = FBlueprintEditorUtils::FindBlueprintForNode(GraphNode))
+			{
+				FPinValueInspectorTreeViewNodePtr RootNode = MakeShared<FPinValueInspectorTreeViewNode>();
+
+				if (FKismetDebugUtilities::FindClassPropertyForPin(Blueprint, GraphPin))
+				{
+					RootNode->DisplayName = GraphPin->GetDisplayName();
+				}
+				else
+				{
+					FFormatNamedArguments Args;
+					Args.Add(TEXT("PinWatchName"), FText::FromString(GraphPin->GetName()));
+					RootNode->DisplayName = FText::Format(LOCTEXT("DisplayNameNoProperty", "{PinWatchName} (no prop)"), Args);
+				}
+
+				RootNode->IconBrush = FBlueprintEditorUtils::GetIconFromPin(GraphPin->PinType);
+				RootNode->IconColor = GraphPin->GetSchema()->GetPinTypeColor(GraphPin->PinType);
+
+				TSharedPtr<FPropertyInstanceInfo> DebugInfo;
+				const FKismetDebugUtilities::EWatchTextResult WatchStatus = FKismetDebugUtilities::GetDebugInfo(DebugInfo, Blueprint, Blueprint->GetObjectBeingDebugged(), GraphPin);
+				switch (WatchStatus)
+				{
+				case FKismetDebugUtilities::EWTR_Valid:
+				{
+					check(DebugInfo.IsValid());
+					const FString ValStr = DebugInfo->Value.ToString();
+					RootNode->Description = FText::FromString(ValStr.Replace(TEXT("\n"), TEXT(" ")));
+					break;
+				}
+
+				case FKismetDebugUtilities::EWTR_NotInScope:
+					RootNode->Description = LOCTEXT("NotInScope", "Not in scope");
+					break;
+
+				case FKismetDebugUtilities::EWTR_NoProperty:
+					RootNode->Description = LOCTEXT("UnknownProperty", "No debug data");
+					break;
+
+				default:
+				case FKismetDebugUtilities::EWTR_NoDebugObject:
+					RootNode->Description = LOCTEXT("NoDebugObject", "No debug object");
+					break;
+				}
+
+				if (DebugInfo.IsValid())
+				{
+					RootNode->Children = MoveTemp(DebugInfo->Children);
+				}
+
+				RootNode->SharedColumnSizeData = MakeShared<FDetailColumnSizeData>();
+				RootNode->SharedColumnSizeData->SetValueColumnWidth(0.5f);
+
+				RootNodes.Add(MoveTemp(RootNode));
+			}
+		}
+	}
+
+	ChildSlot
+	[
+		SNew(SVerticalBox)
+		+SVerticalBox::Slot()
+		.AutoHeight()
+		[
+			SAssignNew(SearchBoxWidget, SSearchBox)
+			.Visibility(this, &SPinValueInspector::GetSearchFilterVisibility)
+		]
+		+SVerticalBox::Slot()
+		[
+			SAssignNew(TreeViewWidget, STreeView<FPinValueInspectorTreeViewNodePtr>)
+			.TreeItemsSource(&RootNodes)
+			.OnGetChildren(this, &SPinValueInspector::OnGetTreeViewNodeChildren)
+			.OnGenerateRow(this, &SPinValueInspector::OnGenerateRowForTreeViewNode)
+			.TreeViewStyle(&FAppStyle::Get().GetWidgetStyle<FTableViewStyle>("PropertyTable.InViewport.ListView"))
+		]
+	];
+
+	// Expand all root-level nodes that have children by default.
+	for (const FPinValueInspectorTreeViewNodePtr& RootNode : RootNodes)
+	{
+		const bool bShouldExpandItem = RootNode->Children.Num() > 0;
+		TreeViewWidget->SetItemExpansion(RootNode, bShouldExpandItem);
+	}
+}
+
+bool SPinValueInspector::ShouldShowSearchFilter() const
+{
+	// Only expose the search filter if child nodes are present (e.g. struct/container types).
+	for (const FPinValueInspectorTreeViewNodePtr& RootNode : RootNodes)
+	{
+		if (RootNode->Children.Num() > 0)
+		{
+			return true;
+		}
+	}
+
+	return false;
+}
+
+EVisibility SPinValueInspector::GetSearchFilterVisibility() const
+{
+	return ShouldShowSearchFilter() ? EVisibility::Visible : EVisibility::Collapsed;
+}
+
+void SPinValueInspector::OnGetTreeViewNodeChildren(FPinValueInspectorTreeViewNodePtr InNode, TArray<FPinValueInspectorTreeViewNodePtr>& OutChildren)
+{
+	for (const TSharedPtr<FPropertyInstanceInfo>& ChildData : InNode->Children)
+	{
+		if (ChildData.IsValid())
+		{
+			FPinValueInspectorTreeViewNodePtr ChildNode = MakeShared<FPinValueInspectorTreeViewNode>();
+
+			ChildNode->DisplayName = ChildData->DisplayName;
+
+			const FString ValStr = ChildData->Value.ToString();
+			ChildNode->Description = FText::FromString(ValStr.Replace(TEXT("\n"), TEXT(" ")));
+
+			FEdGraphPinType PinType;
+			const UEdGraphSchema_K2* K2Schema = GetDefault<UEdGraphSchema_K2>();
+			if (K2Schema->ConvertPropertyToPinType(ChildData->Property.Get(), PinType))
+			{
+				ChildNode->IconBrush = FBlueprintEditorUtils::GetIconFromPin(PinType);
+				ChildNode->IconColor = K2Schema->GetPinTypeColor(PinType);
+			}
+
+			ChildNode->Children = ChildData->Children;
+			ChildNode->SharedColumnSizeData = InNode->SharedColumnSizeData;
+
+			OutChildren.Add(MoveTemp(ChildNode));
+		}
+	}
+}
+
+TSharedRef<ITableRow> SPinValueInspector::OnGenerateRowForTreeViewNode(FPinValueInspectorTreeViewNodePtr InNode, const TSharedRef<STableViewBase>& OwnerTable)
+{
+	// Only show expander arrows if we have at least one child node.
+	bool bShowExpanderArrows = false;
+	for (const FPinValueInspectorTreeViewNodePtr& RootNode : RootNodes)
+	{
+		if (RootNode->Children.Num() > 0)
+		{
+			bShowExpanderArrows = true;
+			break;
+		}
+	}
+
+	return SNew(SPinValueInspector_TreeViewRow, OwnerTable, InNode, bShowExpanderArrows);
+}
+
+#undef LOCTEXT_NAMESPACE
\ No newline at end of file
diff --git a/Engine/Source/Editor/KismetWidgets/Public/SPinValueInspector.h b/Engine/Source/Editor/KismetWidgets/Public/SPinValueInspector.h
new file mode 100644
--- /dev/null
+++ b/Engine/Source/Editor/KismetWidgets/Public/SPinValueInspector.h
@@ -0,0 +1,55 @@
+// Copyright Epic Games, Inc. All Rights Reserved.
+
+#pragma once
+
+#include "CoreMinimal.h"
+#include "EdGraph/EdGraphPin.h"
+#include "Widgets/SCompoundWidget.h"
+#include "Widgets/Views/STreeView.h"
+
+class SSearchBox;
+
+typedef TSharedPtr<struct FPinValueInspectorTreeViewNode> FPinValueInspectorTreeViewNodePtr;
+
+/**
+ * Inspects the referenced pin object's underlying property value and presents it within a tree view.
+ * Compound properties (e.g. structs/containers) will be broken down into a hierarchy of child nodes.
+ */
+class KISMETWIDGETS_API SPinValueInspector : public SCompoundWidget
+{
+public:
+	SLATE_BEGIN_ARGS(SPinValueInspector)
+	{}
+	SLATE_END_ARGS()
+
+	void Construct(const FArguments& InArgs, const FEdGraphPinReference& InPinRef);
+
+	/** Whether the search filter UI should be visible. */
+	bool ShouldShowSearchFilter() const;
+
+protected:
+	/** @return Visibility of the search box filter widget. */
+	EVisibility GetSearchFilterVisibility() const;
+
+	/** Obtains the set of children for the given value item node. */
+	void OnGetTreeViewNodeChildren(FPinValueInspectorTreeViewNodePtr InNode, TArray<FPinValueInspectorTreeViewNodePtr>& OutChildren);
+
+	/** Generates a row widget that presents the given value item node. */
+	TSharedRef<ITableRow> OnGenerateRowForTreeViewNode(FPinValueInspectorTreeViewNodePtr InNode, const TSharedRef<STableViewBase>& OwnerTable);
+
+private:
+	/** Holds a weak reference to the target pin. */
+	FEdGraphPinReference PinRef;
+
+	/** The instance that's currently selected as the debugging target. */
+	TWeakObjectPtr<UObject> TargetObject;
+
+	/** Root node(s) presented through the tree view widget. */
+	TArray<FPinValueInspectorTreeViewNodePtr> RootNodes;
+
+	/** Presents a hierarchical display of the inspected value along with any sub-values as children. */
+	TSharedPtr<STreeView<FPinValueInspectorTreeViewNodePtr>> TreeViewWidget;
+
+	/** Holds a reference to the search box widget, used to filter the tree view display. */
+	TSharedPtr<SSearchBox> SearchBoxWidget;
+};
\ No newline at end of file
diff --git a/Engine/Source/Editor/UnrealEd/Private/Kismet2/KismetDebugUtilities.cpp b/Engine/Source/Editor/UnrealEd/Private/Kismet2/KismetDebugUtilities.cpp
--- a/Engine/Source/Editor/UnrealEd/Private/Kismet2/KismetDebugUtilities.cpp
+++ b/Engine/Source/Editor/UnrealEd/Private/Kismet2/KismetDebugUtilities.cpp
@@ -45,6 +45,14 @@
 
 #define LOCTEXT_NAMESPACE "BlueprintDebugging"
 
+// Temporary console variable to toggle pin value inspection in PIE.
+// @todo - Migrate this to a user-facing editor setting at some point.
+static TAutoConsoleVariable<bool> CVarBPEnablePinValueInspectionDuringPIE(
+	TEXT("BP.EnablePinValueInspectionDuringPIE"),
+	false,
+	TEXT("Enables pin value inspection tooltips during PIE (experimental).")
+);
+
 /** Per-thread data for use by FKismetDebugUtilities functions */
 class FKismetDebugUtilitiesData : public TThreadSingleton<FKismetDebugUtilitiesData>
 {
@@ -1551,6 +1559,49 @@ FKismetDebugUtilities::EWatchTextResult FKismetDebugUtilities::GetWatchText(FStr
 	return Result;
 }
 
+bool FKismetDebugUtilities::CanInspectPinValue(const UEdGraphPin* Pin)
+{
+	if (!CVarBPEnablePinValueInspectionDuringPIE.GetValueOnGameThread())
+	{
+		return false;
+	}
+
+	// Can't inspect the value on an invalid pin object.
+	if (!Pin || Pin->IsPendingKill())
+	{
+		return false;
+	}
+
+	// Can't inspect the value on an orphaned pin object.
+	if (Pin->bOrphanedPin)
+	{
+		return false;
+	}
+
+	// Can't inspect the value on an unknown pin object or if the owning node is disabled.
+	const UEdGraphNode* OwningNode = Pin->GetOwningNodeUnchecked();
+	if (!OwningNode || !OwningNode->IsNodeEnabled())
+	{
+		return false;
+	}
+
+	// Can't inspect exec pins or delegate pins; their values are not defined.
+	const UEdGraphSchema_K2* K2Schema = Cast<UEdGraphSchema_K2>(OwningNode->GetSchema());
+	if (!ensureMsgf(K2Schema, TEXT("Invalid or missing schema.")) || K2Schema->IsExecPin(*Pin) || K2Schema->IsDelegateCategory(Pin->PinType.PinCategory))
+	{
+		return false;
+	}
+
+	// Can't inspect the value if there is no active debug context.
+	const UBlueprint* Blueprint = FBlueprintEditorUtils::FindBlueprintForNode(OwningNode);
+	if (!Blueprint || !Blueprint->GetObjectBeingDebugged())
+	{
+		return false;
+	}
+
+	return true;
+}
+
 FKismetDebugUtilities::EWatchTextResult FKismetDebugUtilities::GetDebugInfo(TSharedPtr<FPropertyInstanceInfo> &OutDebugInfo, UBlueprint* Blueprint, UObject* ActiveObject, const UEdGraphPin* WatchPin)
 {
 	void* DataPtr = nullptr;
diff --git a/Engine/Source/Editor/UnrealEd/Public/Kismet2/KismetDebugUtilities.h b/Engine/Source/Editor/UnrealEd/Public/Kismet2/KismetDebugUtilities.h
--- a/Engine/Source/Editor/UnrealEd/Public/Kismet2/KismetDebugUtilities.h
+++ b/Engine/Source/Editor/UnrealEd/Public/Kismet2/KismetDebugUtilities.h
@@ -300,6 +300,9 @@ public:
 	// Is there debugging data available for this blueprint?
 	static bool HasDebuggingData(const UBlueprint* Blueprint);
 
+	/** Determines if the given pin's value can be inspected */
+	static bool CanInspectPinValue(const UEdGraphPin* Pin);
+
 	/** Returns the breakpoint associated with a node, or NULL */
 	static FBlueprintBreakpoint* FindBreakpointForNode(const UEdGraphNode* OwnerNode, const UBlueprint* OwnerBlueprint, bool bCheckSubLocations = false);
 
