Fix broken search filtering in pin type selector drop-downs when no custom filter is set.

diff --git a/Engine/Source/Editor/KismetWidgets/Private/SPinTypeSelector.cpp b/Engine/Source/Editor/KismetWidgets/Private/SPinTypeSelector.cpp
--- a/Engine/Source/Editor/KismetWidgets/Private/SPinTypeSelector.cpp
+++ b/Engine/Source/Editor/KismetWidgets/Private/SPinTypeSelector.cpp
@@ -1313,11 +1313,8 @@ bool SPinTypeSelector::GetChildrenMatchingSearch(const FText& InSearchText, cons
 		// If children match the search filter, there's no need to do any additional checks
 		if (!bHasChildrenMatchingSearch)
 		{
-			// If valid, attempt to match the custom filter
-			if (CustomFilter.IsValid())
-			{
-				bFilterMatches = CustomFilter->ShouldShowPinTypeTreeItem(Item);
-			}
+			// If valid, attempt to match the custom filter; otherwise, treat it as a match by default.
+			bFilterMatches = CustomFilter.IsValid() ? CustomFilter->ShouldShowPinTypeTreeItem(Item) : true;
 
 			// If we didn't match the custom filter, or it's an empty search, let's not do any checks against the FilterTerms
 			if (bFilterMatches && !bIsEmptySearch)
