[Bug Fix]: Mousing over variables in the blueprint diff tool causes them to jump around
- made the ReadOnly attribute and the ESelectorType::None styling of SPinTypeSelector consistent because they're essentially two ways to say readonly

diff --git a/Engine/Source/Editor/Kismet/Private/SBlueprintPalette.cpp b/Engine/Source/Editor/Kismet/Private/SBlueprintPalette.cpp
--- a/Engine/Source/Editor/Kismet/Private/SBlueprintPalette.cpp
+++ b/Engine/Source/Editor/Kismet/Private/SBlueprintPalette.cpp
@@ -774,7 +774,7 @@ private:
 				.TargetPinType(this, &SPinTypeSelectorHelper::OnGetVarType)
 				.OnPinTypeChanged(this, &SPinTypeSelectorHelper::OnVarTypeChanged)
 				.TypeTreeFilter(ETypeTreeFilter::None)
-				.SelectorType(BlueprintEditorPtr.IsValid() && !bIsDelegate ? SPinTypeSelector::ESelectorType::Partial : SPinTypeSelector::ESelectorType::None)
+				.SelectorType(BlueprintEditorPtr.IsValid() ? SPinTypeSelector::ESelectorType::Partial : SPinTypeSelector::ESelectorType::None)
 				.CustomFilter(CustomPinTypeFilter)
 			];	
 		}
diff --git a/Engine/Source/Editor/KismetWidgets/Private/SPinTypeSelector.cpp b/Engine/Source/Editor/KismetWidgets/Private/SPinTypeSelector.cpp
--- a/Engine/Source/Editor/KismetWidgets/Private/SPinTypeSelector.cpp
+++ b/Engine/Source/Editor/KismetWidgets/Private/SPinTypeSelector.cpp
@@ -351,6 +351,30 @@ void SPinTypeSelector::Construct(const FArguments& InArgs, FGetPinTypeTree GetPi
 
 	// Depending on if this is a compact selector or not, we generate a different compound widget
 	TSharedPtr<SWidget> Widget;
+	const TSharedPtr<SWidget> ReadOnlyWidget = SNew(SHorizontalBox)
+			.Clipping(EWidgetClipping::OnDemand)
+			+ SHorizontalBox::Slot()
+			.VAlign(VAlign_Center)
+			.HAlign(HAlign_Left)
+			.Padding(FMargin(2.0f, 3.0f, 2.0f, 3.0f))
+			.AutoWidth()
+			[
+				// Read-only version does not display container or secondary type separately, so we need to jam it all in the one image
+				SNew(SLayeredImage, TAttribute<const FSlateBrush*>(this, &SPinTypeSelector::GetSecondaryTypeIconImage), TAttribute<FSlateColor>(this, &SPinTypeSelector::GetSecondaryTypeIconColor))
+				.Image(this, &SPinTypeSelector::GetTypeIconImage)
+				.ColorAndOpacity(this, &SPinTypeSelector::GetTypeIconColor)
+			]
+			+ SHorizontalBox::Slot()
+			.Padding(2.0f, 2.0f)
+			.VAlign(VAlign_Center)
+			.HAlign(HAlign_Left)
+			.AutoWidth()
+			[
+				SNew(STextBlock)
+				.Text(this, &SPinTypeSelector::GetTypeDescription, false)
+				.Font(InArgs._Font)
+				.ColorAndOpacity(FSlateColor::UseForeground())
+			];
 
 	if (SelectorType == ESelectorType::Compact)
 	{
@@ -368,8 +392,7 @@ void SPinTypeSelector::Construct(const FArguments& InArgs, FGetPinTypeTree GetPi
 	}
 	else if (SelectorType == ESelectorType::None)
 	{
-		Widget = PrimaryTypeImage.ToSharedRef();
-		Widget->SetToolTipText(TAttribute<FText>::CreateSP(this, &SPinTypeSelector::GetToolTipForComboBoxType));
+		Widget = ReadOnlyWidget;
 	}
 	else if (SelectorType == ESelectorType::Full || SelectorType == ESelectorType::Partial)
 	{
@@ -500,30 +523,7 @@ void SPinTypeSelector::Construct(const FArguments& InArgs, FGetPinTypeTree GetPi
 		]
 		+ SWidgetSwitcher::Slot() // read-only version
 		[
-			SNew(SHorizontalBox)
-			.Clipping(EWidgetClipping::OnDemand)
-			+ SHorizontalBox::Slot()
-			.VAlign(VAlign_Center)
-			.HAlign(HAlign_Left)
-			.Padding(FMargin(2.0f, 3.0f, 2.0f, 3.0f))
-			.AutoWidth()
-			[
-				// Read-only version does not display container or secondary type separately, so we need to jam it all in the one image
-				SNew(SLayeredImage, TAttribute<const FSlateBrush*>(this, &SPinTypeSelector::GetSecondaryTypeIconImage), TAttribute<FSlateColor>(this, &SPinTypeSelector::GetSecondaryTypeIconColor))
-				.Image(this, &SPinTypeSelector::GetTypeIconImage)
-				.ColorAndOpacity(this, &SPinTypeSelector::GetTypeIconColor)
-			]
-			+ SHorizontalBox::Slot()
-			.Padding(2.0f, 2.0f)
-			.VAlign(VAlign_Center)
-			.HAlign(HAlign_Left)
-			.AutoWidth()
-			[
-				SNew(STextBlock)
-				.Text(this, &SPinTypeSelector::GetTypeDescription, false)
-				.Font(InArgs._Font)
-				.ColorAndOpacity(FSlateColor::UseForeground())
-			]
+			ReadOnlyWidget.ToSharedRef()
 		]	
 	];
 }
