UE-150838 - The tooltip for float map values in Blueprint Structures doesn't show precision

Similar to GetToolTipForComboBoxType, GetToolTipForComboBoxSecondaryType should include the subcategory for any value types in a map. This ensures that we can see the precision of the float type when it's a map value.

diff --git a/Engine/Source/Editor/KismetWidgets/Private/SPinTypeSelector.cpp b/Engine/Source/Editor/KismetWidgets/Private/SPinTypeSelector.cpp
--- a/Engine/Source/Editor/KismetWidgets/Private/SPinTypeSelector.cpp
+++ b/Engine/Source/Editor/KismetWidgets/Private/SPinTypeSelector.cpp
@@ -1403,7 +1403,7 @@ FText SPinTypeSelector::GetToolTipForComboBoxSecondaryType() const
 		EditText = LOCTEXT("PinTypeSelector_ValueDisabled", "Cannot edit map value type when they are inherited from parent.");
 	}
 
-	return FText::Format(LOCTEXT("SecondaryTypeTwoLines", "{0}\nValue Type: {1}"), EditText, GetSecondaryTypeDescription());
+	return FText::Format(LOCTEXT("SecondaryTypeTwoLines", "{0}\nValue Type: {1}"), EditText, GetSecondaryTypeDescription(true));
 }
 
 FText SPinTypeSelector::GetToolTipForArrayWidget() const
