Allow pin type selector to be used in a combo box (close self only).

diff --git a/Engine/Source/Editor/KismetWidgets/Private/SPinTypeSelector.cpp b/Engine/Source/Editor/KismetWidgets/Private/SPinTypeSelector.cpp
--- a/Engine/Source/Editor/KismetWidgets/Private/SPinTypeSelector.cpp
+++ b/Engine/Source/Editor/KismetWidgets/Private/SPinTypeSelector.cpp
@@ -1148,7 +1148,7 @@ TSharedRef<SWidget>	SPinTypeSelector::GetMenuContent(bool bForSecondaryType)
 
 TSharedRef<SWidget> SPinTypeSelector::GetPinContainerTypeMenuContent()
 {
-	FMenuBuilder MenuBuilder(true, nullptr);
+	FMenuBuilder MenuBuilder(true, nullptr, nullptr, /*bCloseSelfOnly*/true);
 
 	PinTypeSelectorStatics::PinTypes.Reset();
 	
