Improved BP variable permissions filtering

Allowed multiple filters to be supplied to SPinTypeSelector
Patchup of CL 20892543 for engine merge

diff --git a/Engine/Plugins/Animation/ControlRig/Source/ControlRigEditor/Private/ControlRigGraphDetails.cpp b/Engine/Plugins/Animation/ControlRig/Source/ControlRigEditor/Private/ControlRigGraphDetails.cpp
--- a/Engine/Plugins/Animation/ControlRig/Source/ControlRigEditor/Private/ControlRigGraphDetails.cpp
+++ b/Engine/Plugins/Animation/ControlRig/Source/ControlRigEditor/Private/ControlRigGraphDetails.cpp
@@ -177,9 +177,14 @@ public:
 		
 		if (ControlRigEditorPtr.IsValid())
 		{
-			if (ControlRigEditorPtr.Pin()->GetImportedPinTypeSelectorFilter().IsValid())
+			TArray<TSharedPtr<IPinTypeSelectorFilter>> Filters;
+			ControlRigEditorPtr.Pin()->GetPinTypeSelectorFilters(Filters);
+			for(const TSharedPtr<IPinTypeSelectorFilter>& Filter : Filters)
 			{
-				return ControlRigEditorPtr.Pin()->GetImportedPinTypeSelectorFilter()->ShouldShowPinTypeTreeItem(InItem);				
+				if(!Filter->ShouldShowPinTypeTreeItem(InItem))
+				{
+					return false;
+				}
 			}
 		}
 
@@ -200,10 +205,10 @@ void FControlRigArgumentLayout::GenerateHeaderRowContent(FDetailWidgetRow& NodeR
 	ETypeTreeFilter TypeTreeFilter = ETypeTreeFilter::None;
 	TypeTreeFilter |= ETypeTreeFilter::AllowExec;
 
-	TSharedPtr<IPinTypeSelectorFilter> CustomPinTypeFilter;
+	TArray<TSharedPtr<IPinTypeSelectorFilter>> CustomPinTypeFilters;
 	if (ControlRigEditorPtr.IsValid())
 	{
-		CustomPinTypeFilter = MakeShared<FControlRigArgumentPinTypeSelectorFilter>(ControlRigEditorPtr, GraphPtr);
+		CustomPinTypeFilters.Add(MakeShared<FControlRigArgumentPinTypeSelectorFilter>(ControlRigEditorPtr, GraphPtr));
 	}
 	
 	NodeRow
@@ -257,7 +262,7 @@ void FControlRigArgumentLayout::GenerateHeaderRowContent(FDetailWidgetRow& NodeR
 				.TypeTreeFilter(TypeTreeFilter)
 				.bAllowArrays(!ShouldPinBeReadOnly())
 				.IsEnabled(!ShouldPinBeReadOnly(true))
-				.CustomFilter(CustomPinTypeFilter)
+				.CustomFilters(CustomPinTypeFilters)
 				.Font(IDetailLayoutBuilder::GetDetailFont())
 			]
 			+ SHorizontalBox::Slot()
diff --git a/Engine/Plugins/Animation/ControlRig/Source/ControlRigEditor/Private/ControlRigLocalVariableDetails.cpp b/Engine/Plugins/Animation/ControlRig/Source/ControlRigEditor/Private/ControlRigLocalVariableDetails.cpp
--- a/Engine/Plugins/Animation/ControlRig/Source/ControlRigEditor/Private/ControlRigLocalVariableDetails.cpp
+++ b/Engine/Plugins/Animation/ControlRig/Source/ControlRigEditor/Private/ControlRigLocalVariableDetails.cpp
@@ -65,7 +65,6 @@ void FRigVMLocalVariableDetails::CustomizeDetails(IDetailLayoutBuilder& DetailBu
 		.OnVerifyTextChanged(this, &FRigVMLocalVariableDetails::OnVerifyNameChanged)
 	];
 
-	TSharedPtr<IPinTypeSelectorFilter> CustomPinTypeFilter;
 	Category.AddCustomRow(LOCTEXT("VariableTypeLabel", "Variable Type"))
 		.NameContent()
 		[
@@ -82,7 +81,6 @@ void FRigVMLocalVariableDetails::CustomizeDetails(IDetailLayoutBuilder& DetailBu
 			.Schema(Schema)
 			.TypeTreeFilter(ETypeTreeFilter::None)
 			.Font(DetailFontInfo)
-			.CustomFilter(CustomPinTypeFilter)
 		];
 
 
diff --git a/Engine/Source/Editor/Kismet/Private/BlueprintDetailsCustomization.cpp b/Engine/Source/Editor/Kismet/Private/BlueprintDetailsCustomization.cpp
--- a/Engine/Source/Editor/Kismet/Private/BlueprintDetailsCustomization.cpp
+++ b/Engine/Source/Editor/Kismet/Private/BlueprintDetailsCustomization.cpp
@@ -340,10 +340,10 @@ void FBlueprintVarActionDetails::CustomizeDetails( IDetailLayoutBuilder& DetailL
 
 	TSharedPtr<SToolTip> VarTypeTooltip = IDocumentation::Get()->CreateToolTip(LOCTEXT("VarTypeTooltip", "The type of the variable."), NULL, DocLink, TEXT("VariableType"));
 
-	TSharedPtr<IPinTypeSelectorFilter> CustomPinTypeFilter;
+	TArray<TSharedPtr<IPinTypeSelectorFilter>> CustomPinTypeFilters;
 	if (BlueprintEditor.IsValid())
 	{
-		CustomPinTypeFilter = BlueprintEditor.Pin()->GetImportedPinTypeSelectorFilter();
+		BlueprintEditor.Pin()->GetPinTypeSelectorFilters(CustomPinTypeFilters);
 	}
 	
 	const UEdGraphSchema* Schema = GetDefault<UEdGraphSchema_K2>();
@@ -374,7 +374,7 @@ void FBlueprintVarActionDetails::CustomizeDetails( IDetailLayoutBuilder& DetailL
 			.TypeTreeFilter(ETypeTreeFilter::None)
 			.Font(DetailFontInfo)
 			.ToolTip(VarTypeTooltip)
-			.CustomFilter(CustomPinTypeFilter)
+			.CustomFilters(CustomPinTypeFilters)
 		];
 
 	TSharedPtr<SToolTip> ToolTipTooltip = IDocumentation::Get()->CreateToolTip(LOCTEXT("VarToolTipTooltip", "Extra information about this variable, shown when cursor is over it."), NULL, DocLink, TEXT("Description"));
@@ -3109,7 +3109,7 @@ void FBlueprintGraphArgumentLayout::GenerateHeaderRowContent( FDetailWidgetRow&
 		TypeTreeFilter |= ETypeTreeFilter::AllowWildcard;
 	}
 
-	TSharedPtr<IPinTypeSelectorFilter> CustomPinTypeFilter;
+	TArray<TSharedPtr<IPinTypeSelectorFilter>> CustomPinTypeFilters;
 	if (GraphActionDetailsPtr.IsValid())
 	{
 		TSharedPtr<SMyBlueprint> MyBlueprintPtr = GraphActionDetailsPtr.Pin()->GetMyBlueprint().Pin();
@@ -3118,7 +3118,7 @@ void FBlueprintGraphArgumentLayout::GenerateHeaderRowContent( FDetailWidgetRow&
 			TSharedPtr<FBlueprintEditor> BlueprintEditorPtr = MyBlueprintPtr->GetBlueprintEditor().Pin();
 			if (BlueprintEditorPtr.IsValid())
 			{
-				CustomPinTypeFilter = BlueprintEditorPtr->GetImportedPinTypeSelectorFilter();
+				BlueprintEditorPtr->GetPinTypeSelectorFilters(CustomPinTypeFilters);
 			}
 		}
 	}
@@ -3162,7 +3162,7 @@ void FBlueprintGraphArgumentLayout::GenerateHeaderRowContent( FDetailWidgetRow&
 				.bAllowArrays(!ShouldPinBeReadOnly())
 				.IsEnabled(!ShouldPinBeReadOnly(true))
 				.Font( IDetailLayoutBuilder::GetDetailFont() )
-				.CustomFilter(CustomPinTypeFilter)
+				.CustomFilters(CustomPinTypeFilters)
 		]
 		+ SHorizontalBox::Slot()
 		.HAlign(HAlign_Right)
diff --git a/Engine/Source/Editor/Kismet/Private/BlueprintEditor.cpp b/Engine/Source/Editor/Kismet/Private/BlueprintEditor.cpp
--- a/Engine/Source/Editor/Kismet/Private/BlueprintEditor.cpp
+++ b/Engine/Source/Editor/Kismet/Private/BlueprintEditor.cpp
@@ -397,6 +397,59 @@ namespace BlueprintEditorImpl
 		bool bIsFilterEnabled;
 	};
 
+	/** Pin type filter for asset permissions */
+	class FPermissionsPinTypeSelectorFilter : public IPinTypeSelectorFilter, public TSharedFromThis<FPermissionsPinTypeSelectorFilter>
+	{
+	public:
+		FPermissionsPinTypeSelectorFilter(TConstArrayView<UBlueprint*> InBlueprints)
+		{
+			FAssetReferenceFilterContext Context;
+			Context.ReferencingAssets.Reserve(InBlueprints.Num());
+
+			for (UBlueprint* Blueprint : InBlueprints)
+			{
+				Context.ReferencingAssets.Add(FAssetData(Blueprint));
+			}
+
+			AssetReferenceFilter = GEditor->MakeAssetReferenceFilter(Context);
+		}
+
+		// IPinTypeSelectorFilter interface
+		virtual bool ShouldShowPinTypeTreeItem(FPinTypeTreeItem InItem) const override
+		{
+			const FSoftObjectPath& AssetReference = InItem->GetSubCategoryObjectAsset();
+			FTopLevelAssetPath TopLevelAssetPath;
+
+			if (AssetReference.IsAsset())
+			{
+				TopLevelAssetPath = FTopLevelAssetPath(AssetReference.GetLongPackageFName(), *AssetReference.GetAssetName());
+			}
+
+			// First check pin type permissions
+			if (!FBlueprintActionDatabase::IsPinTypeAllowed(InItem->GetPinType(false), TopLevelAssetPath))
+			{
+				return false;
+			}
+
+			// Then asset permissions
+			if(AssetReferenceFilter.IsValid() && AssetReference.IsValid())
+			{
+				FAssetRegistryModule& AssetRegistryModule = FModuleManager::LoadModuleChecked<FAssetRegistryModule>("AssetRegistry");
+				FAssetData AssetData = AssetRegistryModule.Get().GetAssetByObjectPath(AssetReference.GetAssetPathName());
+				if (!AssetReferenceFilter->PassesFilter(AssetData))
+				{
+					return false;
+				}
+			}
+
+			return true;
+		}
+
+	private:
+		/** Filter for asset references */
+		TSharedPtr<IAssetReferenceFilter> AssetReferenceFilter;
+	};
+
 	/**
 	 * Utility function that will check to see if the specified graph has any 
 	 * nodes other than those that come default, pre-placed, in the graph.
@@ -927,6 +980,12 @@ void FBlueprintEditor::AnalyticsTrackCompileEvent( UBlueprint* Blueprint, int32
 	}
 }
 
+void FBlueprintEditor::GetPinTypeSelectorFilters(TArray<TSharedPtr<IPinTypeSelectorFilter>>& OutFilters) const
+{ 
+	OutFilters.Add(ImportedPinTypeSelectorFilter);
+	OutFilters.Add(PermissionsPinTypeSelectorFilter);
+}
+
 void FBlueprintEditor::RefreshEditors(ERefreshBlueprintEditorReason::Type Reason)
 {
 	bool bForceFocusOnSelectedNodes = false;
@@ -1890,6 +1949,7 @@ void FBlueprintEditor::CommonInitialization(const TArray<UBlueprint*>& InitBluep
 	// Create imported namespace type filters for value editing.
 	ImportedClassViewerFilter = MakeShared<BlueprintEditorImpl::FImportedClassViewerFilterProxy>(ImportedNamespaceHelper->GetClassViewerFilter());
 	ImportedPinTypeSelectorFilter = MakeShared<BlueprintEditorImpl::FImportedPinTypeSelectorFilterProxy>(ImportedNamespaceHelper->GetPinTypeSelectorFilter());
+	PermissionsPinTypeSelectorFilter = MakeShared<BlueprintEditorImpl::FPermissionsPinTypeSelectorFilter>(InitBlueprints);
 
 	// Make sure we know when tabs become active to update details tab
 	OnActiveTabChangedDelegateHandle = FGlobalTabmanager::Get()->OnActiveTabChanged_Subscribe( FOnActiveTabChanged::FDelegate::CreateRaw(this, &FBlueprintEditor::OnActiveTabChanged) );
diff --git a/Engine/Source/Editor/Kismet/Private/SBlueprintPalette.cpp b/Engine/Source/Editor/Kismet/Private/SBlueprintPalette.cpp
--- a/Engine/Source/Editor/Kismet/Private/SBlueprintPalette.cpp
+++ b/Engine/Source/Editor/Kismet/Private/SBlueprintPalette.cpp
@@ -754,10 +754,10 @@ private:
 	{
 		OnTypeChanged = InArgs._OnTypeChanged;
 		
-		TSharedPtr<IPinTypeSelectorFilter> CustomPinTypeFilter;
+		TArray<TSharedPtr<IPinTypeSelectorFilter>> CustomPinTypeFilters;
 		if (BlueprintEditorPtr.IsValid())
 		{
-			CustomPinTypeFilter = BlueprintEditorPtr.Pin()->GetImportedPinTypeSelectorFilter();
+			BlueprintEditorPtr.Pin()->GetPinTypeSelectorFilters(CustomPinTypeFilters);
 		}
 
 		const UEdGraphSchema* Schema = GetDefault<UEdGraphSchema_K2>();
@@ -784,7 +784,7 @@ private:
 				.OnPinTypeChanged(this, &SPinTypeSelectorHelper::OnVarTypeChanged)
 				.TypeTreeFilter(ETypeTreeFilter::None)
 				.SelectorType(BlueprintEditorPtr.IsValid() ? SPinTypeSelector::ESelectorType::Partial : SPinTypeSelector::ESelectorType::None)
-				.CustomFilter(CustomPinTypeFilter)
+				.CustomFilters(CustomPinTypeFilters)
 			];	
 		}
 	}
diff --git a/Engine/Source/Editor/Kismet/Public/BlueprintEditor.h b/Engine/Source/Editor/Kismet/Public/BlueprintEditor.h
--- a/Engine/Source/Editor/Kismet/Public/BlueprintEditor.h
+++ b/Engine/Source/Editor/Kismet/Public/BlueprintEditor.h
@@ -244,7 +244,9 @@ public:
 	virtual void AnalyticsTrackNodeEvent(UBlueprint* Blueprint, UEdGraphNode *GraphNode, bool bNodeDelete = false) const override;
 	void AnalyticsTrackCompileEvent(UBlueprint* Blueprint, int32 NumErrors, int32 NumWarnings) const;
 	virtual TSharedPtr<class IClassViewerFilter> GetImportedClassViewerFilter() const override { return ImportedClassViewerFilter; }
+	UE_DEPRECATED(5.1, "Please use GetPinTypeSelectorFilters")
 	virtual TSharedPtr<class IPinTypeSelectorFilter> GetImportedPinTypeSelectorFilter() const override { return ImportedPinTypeSelectorFilter; }
+	virtual void GetPinTypeSelectorFilters(TArray<TSharedPtr<class IPinTypeSelectorFilter>>& OutFilters) const override;
 	virtual bool IsNonImportedObject(const UObject* InObject) const;
 	virtual bool IsNonImportedObject(const FSoftObjectPath& InObject) const;
 	//~ End IBlueprintEditor Interface
@@ -1415,6 +1417,12 @@ protected:
 	/** Filter used to restrict pin type selector widgets in the editor context to imported namespaces only */
 	TSharedPtr<class IPinTypeSelectorFilter> ImportedPinTypeSelectorFilter;
 
+	/** Filter used to restrict pin type selector widgets to those which permissions allow */
+	TSharedPtr<class IPinTypeSelectorFilter> PermissionsPinTypeSelectorFilter;
+
+	/** Filter used to compose multiple pin type filters together */
+	TSharedPtr<class IPinTypeSelectorFilter> CompositePinTypeSelectorFilter;
+	
 	FOnSetPinVisibility OnSetPinVisibility;
 
 	/** Has someone requested a deferred update of the saved document state? */
diff --git a/Engine/Source/Editor/Kismet/Public/BlueprintEditorModule.h b/Engine/Source/Editor/Kismet/Public/BlueprintEditorModule.h
--- a/Engine/Source/Editor/Kismet/Public/BlueprintEditorModule.h
+++ b/Engine/Source/Editor/Kismet/Public/BlueprintEditorModule.h
@@ -109,8 +109,12 @@ public:
 	virtual TSharedPtr<class IClassViewerFilter> GetImportedClassViewerFilter() const { return nullptr; }
 
 	/** Return the pin type selector filter associated with the current set of imported namespaces within this editor context. Default is NULL (no filter). */
+	UE_DEPRECATED(5.1, "Please use GetPinTypeSelectorFilters")
 	virtual TSharedPtr<class IPinTypeSelectorFilter> GetImportedPinTypeSelectorFilter() const { return nullptr; }
 
+	/** Get all the the pin type selector filters within this editor context. */
+	virtual void GetPinTypeSelectorFilters(TArray<TSharedPtr<class IPinTypeSelectorFilter>>& OutFilters) const {}
+	
 	/** Return whether the given object falls outside the scope of the current set of imported namespaces within this editor context. Default is FALSE (imported). */
 	virtual bool IsNonImportedObject(const UObject* InObject) const { return false; }
 
diff --git a/Engine/Source/Editor/KismetWidgets/Private/SPinTypeSelector.cpp b/Engine/Source/Editor/KismetWidgets/Private/SPinTypeSelector.cpp
--- a/Engine/Source/Editor/KismetWidgets/Private/SPinTypeSelector.cpp
+++ b/Engine/Source/Editor/KismetWidgets/Private/SPinTypeSelector.cpp
@@ -314,14 +314,22 @@ void SPinTypeSelector::Construct(const FArguments& InArgs, FGetPinTypeTree GetPi
 	SelectorType = InArgs._SelectorType;
 
 	NumFilteredPinTypeItems = 0;
-	if (InArgs._CustomFilter.IsValid())
+
+	if(InArgs._CustomFilters.Num() > 0)
 	{
-		CustomFilter = MakeShared<FPinTypeSelectorCustomFilterProxy>(InArgs._CustomFilter.ToSharedRef(), FSimpleDelegate::CreateSP(this, &SPinTypeSelector::OnCustomFilterChanged));
+		for(const TSharedPtr<IPinTypeSelectorFilter>& Filter : InArgs._CustomFilters)
+		{
+			CustomFilters.Add(MakeShared<FPinTypeSelectorCustomFilterProxy>(Filter.ToSharedRef(), FSimpleDelegate::CreateSP(this, &SPinTypeSelector::OnCustomFilterChanged)));
+		}
+	}
+	else if (InArgs._CustomFilter.IsValid())
+	{
+		CustomFilters.Add(MakeShared<FPinTypeSelectorCustomFilterProxy>(InArgs._CustomFilter.ToSharedRef(), FSimpleDelegate::CreateSP(this, &SPinTypeSelector::OnCustomFilterChanged)));
 	}
 	else if (UClass* PinTypeSelectorFilterClass = GetDefault<UPinTypeSelectorFilter>()->FilterClass.LoadSynchronous() )
 	{
 		TSharedPtr<IPinTypeSelectorFilter> SelectorFilter = GetDefault<UPinTypeSelectorFilter>(PinTypeSelectorFilterClass)->GetPinTypeSelectorFilter();
-		CustomFilter = MakeShared<FPinTypeSelectorCustomFilterProxy>(SelectorFilter.ToSharedRef(), FSimpleDelegate::CreateSP(this, &SPinTypeSelector::OnCustomFilterChanged));
+		CustomFilters.Add(MakeShared<FPinTypeSelectorCustomFilterProxy>(SelectorFilter.ToSharedRef(), FSimpleDelegate::CreateSP(this, &SPinTypeSelector::OnCustomFilterChanged)));
 	}
 
 	bIsRightMousePressed = false;
@@ -1036,7 +1044,7 @@ TSharedRef<SWidget>	SPinTypeSelector::GetMenuContent(bool bForSecondaryType)
 	GetChildrenWithSupportedTypes(TypeTreeRoot, FilteredBySupportedTypes);
 	TypeTreeRoot = FilteredBySupportedTypes;
 
-	if (CustomFilter.IsValid())
+	if (CustomFilters.Num() > 0)
 	{
 		NumFilteredPinTypeItems = 0;
 		FilteredTypeTreeRoot.Empty();
@@ -1063,12 +1071,21 @@ TSharedRef<SWidget>	SPinTypeSelector::GetMenuContent(bool bForSecondaryType)
 			.OnTextChanged( this, &SPinTypeSelector::OnFilterTextChanged )
 			.OnTextCommitted( this, &SPinTypeSelector::OnFilterTextCommitted );
 
-		TSharedPtr<SWidget> CustomFilterOptionsWidget;
-		if (CustomFilter.IsValid())
+		TArray<TSharedPtr<SWidget>> CustomFilterOptionsWidgets;
+		for(const TSharedPtr<IPinTypeSelectorFilter>& CustomFilter : CustomFilters)
 		{
-			CustomFilterOptionsWidget = CustomFilter->GetFilterOptionsWidget();
+			if(CustomFilter.IsValid())
+			{
+				TSharedPtr<SWidget> CustomWidget = CustomFilter->GetFilterOptionsWidget();
+				if(CustomWidget.IsValid())
+				{
+					CustomFilterOptionsWidgets.Add(CustomWidget);
+				}
+			}
 		}
 
+		TSharedPtr<SHorizontalBox> CustomWidgetContainer;
+		
 		MenuContent = SAssignNew(PinTypeSelectorMenuOwner, SMenuOwner)
 			[
 				SNew(SListViewSelectorDropdownMenu<FPinTypeTreeItem>, FilterTextBox, TypeTreeView)
@@ -1096,9 +1113,9 @@ TSharedRef<SWidget>	SPinTypeSelector::GetMenuContent(bool bForSecondaryType)
 					.Padding(8.f, 0.f, 8.f, 4.f)
 					[
 						SNew(SBox)
-						.Visibility(CustomFilter.IsValid() ? EVisibility::Visible : EVisibility::Collapsed)
+						.Visibility(CustomFilterOptionsWidgets.Num() > 0 ? EVisibility::Visible : EVisibility::Collapsed)
 						[
-							SNew(SHorizontalBox)
+							SAssignNew(CustomWidgetContainer, SHorizontalBox)
 							+SHorizontalBox::Slot()
 							.VAlign(VAlign_Center)
 							.FillWidth(1.f)
@@ -1106,17 +1123,24 @@ TSharedRef<SWidget>	SPinTypeSelector::GetMenuContent(bool bForSecondaryType)
 								SNew(STextBlock)
 								.Text(this, &SPinTypeSelector::GetPinTypeItemCountText)
 							]
-							+SHorizontalBox::Slot()
-							.VAlign(VAlign_Center)
-							.AutoWidth()
-							[
-								CustomFilterOptionsWidget.IsValid() ? CustomFilterOptionsWidget.ToSharedRef() : SNullWidget::NullWidget
-							]
+							
 						]
 					]
 				]
 			];
-			
+
+		if(CustomWidgetContainer.IsValid())
+		{
+			for(const TSharedPtr<SWidget>& CustomFilterOptionsWidget : CustomFilterOptionsWidgets)
+			{
+				CustomWidgetContainer->AddSlot()
+				.VAlign(VAlign_Center)
+				.AutoWidth()
+				[
+					CustomFilterOptionsWidget.ToSharedRef()
+				];
+			}
+		}
 
 		if (bForSecondaryType)
 		{
@@ -1314,8 +1338,14 @@ bool SPinTypeSelector::GetChildrenMatchingSearch(const FText& InSearchText, cons
 		// If children match the search filter, there's no need to do any additional checks
 		if (!bHasChildrenMatchingSearch)
 		{
-			// If valid, attempt to match the custom filter; otherwise, treat it as a match by default.
-			bFilterMatches = CustomFilter.IsValid() ? CustomFilter->ShouldShowPinTypeTreeItem(Item) : true;
+			// If valid, attempt to match the custom filters; otherwise, treat it as a match by default.
+			for(const TSharedPtr<IPinTypeSelectorFilter>& CustomFilter : CustomFilters)
+			{
+				if(CustomFilter.IsValid())
+				{
+					bFilterMatches &= CustomFilter->ShouldShowPinTypeTreeItem(Item);
+				}
+			}
 
 			// If we didn't match the custom filter, or it's an empty search, let's not do any checks against the FilterTerms
 			if (bFilterMatches && !bIsEmptySearch)
diff --git a/Engine/Source/Editor/KismetWidgets/Public/SPinTypeSelector.h b/Engine/Source/Editor/KismetWidgets/Public/SPinTypeSelector.h
--- a/Engine/Source/Editor/KismetWidgets/Public/SPinTypeSelector.h
+++ b/Engine/Source/Editor/KismetWidgets/Public/SPinTypeSelector.h
@@ -82,7 +82,8 @@ public:
 		SLATE_ATTRIBUTE( FSlateFontInfo, Font )
 		SLATE_ARGUMENT( ESelectorType, SelectorType )
 		SLATE_ATTRIBUTE(bool, ReadOnly)
-		SLATE_ARGUMENT(TSharedPtr<class IPinTypeSelectorFilter>, CustomFilter)
+		SLATE_ARGUMENT_DEPRECATED(TSharedPtr<class IPinTypeSelectorFilter>, CustomFilter, 5.1, "Please use CustomFilters instead")
+		SLATE_ARGUMENT(TArray<TSharedPtr<class IPinTypeSelectorFilter>>, CustomFilters)
 	SLATE_END_ARGS()
 public:
 	void Construct(const FArguments& InArgs, FGetPinTypeTree GetPinTypeTreeFunc);
@@ -166,7 +167,7 @@ protected:
 	TWeakPtr<SMenuOwner> PinTypeSelectorMenuOwner;
 
 	/** An interface to optionally apply a custom filter to the available pin type items for display. */
-	TSharedPtr<class IPinTypeSelectorFilter> CustomFilter;
+	TArray<TSharedPtr<class IPinTypeSelectorFilter>> CustomFilters;
 
 	/** Array checkbox support functions */
 	ECheckBoxState IsArrayChecked() const;
