[Backout] - CL21088562
Original CL Desc
-----------------------------------------------------------------
Improved BP variable permissions filtering

Allowed multiple filters to be supplied to SPinTypeSelector
Patchup of CL 20892543 for engine merge

diff --git a/Engine/Plugins/Animation/ControlRig/Source/ControlRigEditor/Private/ControlRigGraphDetails.cpp b/Engine/Plugins/Animation/ControlRig/Source/ControlRigEditor/Private/ControlRigGraphDetails.cpp
--- a/Engine/Plugins/Animation/ControlRig/Source/ControlRigEditor/Private/ControlRigGraphDetails.cpp
+++ b/Engine/Plugins/Animation/ControlRig/Source/ControlRigEditor/Private/ControlRigGraphDetails.cpp
@@ -177,14 +177,9 @@ public:
 		
 		if (ControlRigEditorPtr.IsValid())
 		{
-			TArray<TSharedPtr<IPinTypeSelectorFilter>> Filters;
-			ControlRigEditorPtr.Pin()->GetPinTypeSelectorFilters(Filters);
-			for(const TSharedPtr<IPinTypeSelectorFilter>& Filter : Filters)
+			if (ControlRigEditorPtr.Pin()->GetImportedPinTypeSelectorFilter().IsValid())
 			{
-				if(!Filter->ShouldShowPinTypeTreeItem(InItem))
-				{
-					return false;
-				}
+				return ControlRigEditorPtr.Pin()->GetImportedPinTypeSelectorFilter()->ShouldShowPinTypeTreeItem(InItem);				
 			}
 		}
 
@@ -205,10 +200,10 @@ void FControlRigArgumentLayout::GenerateHeaderRowContent(FDetailWidgetRow& NodeR
 	ETypeTreeFilter TypeTreeFilter = ETypeTreeFilter::None;
 	TypeTreeFilter |= ETypeTreeFilter::AllowExec;
 
-	TArray<TSharedPtr<IPinTypeSelectorFilter>> CustomPinTypeFilters;
+	TSharedPtr<IPinTypeSelectorFilter> CustomPinTypeFilter;
 	if (ControlRigEditorPtr.IsValid())
 	{
-		CustomPinTypeFilters.Add(MakeShared<FControlRigArgumentPinTypeSelectorFilter>(ControlRigEditorPtr, GraphPtr));
+		CustomPinTypeFilter = MakeShared<FControlRigArgumentPinTypeSelectorFilter>(ControlRigEditorPtr, GraphPtr);
 	}
 	
 	NodeRow
@@ -262,7 +257,7 @@ void FControlRigArgumentLayout::GenerateHeaderRowContent(FDetailWidgetRow& NodeR
 				.TypeTreeFilter(TypeTreeFilter)
 				.bAllowArrays(!ShouldPinBeReadOnly())
 				.IsEnabled(!ShouldPinBeReadOnly(true))
-				.CustomFilters(CustomPinTypeFilters)
+				.CustomFilter(CustomPinTypeFilter)
 				.Font(IDetailLayoutBuilder::GetDetailFont())
 			]
 			+ SHorizontalBox::Slot()
diff --git a/Engine/Plugins/Animation/ControlRig/Source/ControlRigEditor/Private/ControlRigLocalVariableDetails.cpp b/Engine/Plugins/Animation/ControlRig/Source/ControlRigEditor/Private/ControlRigLocalVariableDetails.cpp
--- a/Engine/Plugins/Animation/ControlRig/Source/ControlRigEditor/Private/ControlRigLocalVariableDetails.cpp
+++ b/Engine/Plugins/Animation/ControlRig/Source/ControlRigEditor/Private/ControlRigLocalVariableDetails.cpp
@@ -65,6 +65,7 @@ void FRigVMLocalVariableDetails::CustomizeDetails(IDetailLayoutBuilder& DetailBu
 		.OnVerifyTextChanged(this, &FRigVMLocalVariableDetails::OnVerifyNameChanged)
 	];
 
+	TSharedPtr<IPinTypeSelectorFilter> CustomPinTypeFilter;
 	Category.AddCustomRow(LOCTEXT("VariableTypeLabel", "Variable Type"))
 		.NameContent()
 		[
@@ -81,6 +82,7 @@ void FRigVMLocalVariableDetails::CustomizeDetails(IDetailLayoutBuilder& DetailBu
 			.Schema(Schema)
 			.TypeTreeFilter(ETypeTreeFilter::None)
 			.Font(DetailFontInfo)
+			.CustomFilter(CustomPinTypeFilter)
 		];
 
 
diff --git a/Engine/Source/Editor/Kismet/Private/BlueprintDetailsCustomization.cpp b/Engine/Source/Editor/Kismet/Private/BlueprintDetailsCustomization.cpp
--- a/Engine/Source/Editor/Kismet/Private/BlueprintDetailsCustomization.cpp
+++ b/Engine/Source/Editor/Kismet/Private/BlueprintDetailsCustomization.cpp
@@ -340,10 +340,10 @@ void FBlueprintVarActionDetails::CustomizeDetails( IDetailLayoutBuilder& DetailL
 
 	TSharedPtr<SToolTip> VarTypeTooltip = IDocumentation::Get()->CreateToolTip(LOCTEXT("VarTypeTooltip", "The type of the variable."), NULL, DocLink, TEXT("VariableType"));
 
-	TArray<TSharedPtr<IPinTypeSelectorFilter>> CustomPinTypeFilters;
+	TSharedPtr<IPinTypeSelectorFilter> CustomPinTypeFilter;
 	if (BlueprintEditor.IsValid())
 	{
-		BlueprintEditor.Pin()->GetPinTypeSelectorFilters(CustomPinTypeFilters);
+		CustomPinTypeFilter = BlueprintEditor.Pin()->GetImportedPinTypeSelectorFilter();
 	}
 	
 	const UEdGraphSchema* Schema = GetDefault<UEdGraphSchema_K2>();
@@ -374,7 +374,7 @@ void FBlueprintVarActionDetails::CustomizeDetails( IDetailLayoutBuilder& DetailL
 			.TypeTreeFilter(ETypeTreeFilter::None)
 			.Font(DetailFontInfo)
 			.ToolTip(VarTypeTooltip)
-			.CustomFilters(CustomPinTypeFilters)
+			.CustomFilter(CustomPinTypeFilter)
 		];
 
 	TSharedPtr<SToolTip> ToolTipTooltip = IDocumentation::Get()->CreateToolTip(LOCTEXT("VarToolTipTooltip", "Extra information about this variable, shown when cursor is over it."), NULL, DocLink, TEXT("Description"));
@@ -3109,7 +3109,7 @@ void FBlueprintGraphArgumentLayout::GenerateHeaderRowContent( FDetailWidgetRow&
 		TypeTreeFilter |= ETypeTreeFilter::AllowWildcard;
 	}
 
-	TArray<TSharedPtr<IPinTypeSelectorFilter>> CustomPinTypeFilters;
+	TSharedPtr<IPinTypeSelectorFilter> CustomPinTypeFilter;
 	if (GraphActionDetailsPtr.IsValid())
 	{
 		TSharedPtr<SMyBlueprint> MyBlueprintPtr = GraphActionDetailsPtr.Pin()->GetMyBlueprint().Pin();
@@ -3118,7 +3118,7 @@ void FBlueprintGraphArgumentLayout::GenerateHeaderRowContent( FDetailWidgetRow&
 			TSharedPtr<FBlueprintEditor> BlueprintEditorPtr = MyBlueprintPtr->GetBlueprintEditor().Pin();
 			if (BlueprintEditorPtr.IsValid())
 			{
-				BlueprintEditorPtr->GetPinTypeSelectorFilters(CustomPinTypeFilters);
+				CustomPinTypeFilter = BlueprintEditorPtr->GetImportedPinTypeSelectorFilter();
 			}
 		}
 	}
@@ -3162,7 +3162,7 @@ void FBlueprintGraphArgumentLayout::GenerateHeaderRowContent( FDetailWidgetRow&
 				.bAllowArrays(!ShouldPinBeReadOnly())
 				.IsEnabled(!ShouldPinBeReadOnly(true))
 				.Font( IDetailLayoutBuilder::GetDetailFont() )
-				.CustomFilters(CustomPinTypeFilters)
+				.CustomFilter(CustomPinTypeFilter)
 		]
 		+ SHorizontalBox::Slot()
 		.HAlign(HAlign_Right)
diff --git a/Engine/Source/Editor/Kismet/Private/BlueprintEditor.cpp b/Engine/Source/Editor/Kismet/Private/BlueprintEditor.cpp
--- a/Engine/Source/Editor/Kismet/Private/BlueprintEditor.cpp
+++ b/Engine/Source/Editor/Kismet/Private/BlueprintEditor.cpp
@@ -397,59 +397,6 @@ namespace BlueprintEditorImpl
 		bool bIsFilterEnabled;
 	};
 
-	/** Pin type filter for asset permissions */
-	class FPermissionsPinTypeSelectorFilter : public IPinTypeSelectorFilter, public TSharedFromThis<FPermissionsPinTypeSelectorFilter>
-	{
-	public:
-		FPermissionsPinTypeSelectorFilter(TConstArrayView<UBlueprint*> InBlueprints)
-		{
-			FAssetReferenceFilterContext Context;
-			Context.ReferencingAssets.Reserve(InBlueprints.Num());
-
-			for (UBlueprint* Blueprint : InBlueprints)
-			{
-				Context.ReferencingAssets.Add(FAssetData(Blueprint));
-			}
-
-			AssetReferenceFilter = GEditor->MakeAssetReferenceFilter(Context);
-		}
-
-		// IPinTypeSelectorFilter interface
-		virtual bool ShouldShowPinTypeTreeItem(FPinTypeTreeItem InItem) const override
-		{
-			const FSoftObjectPath& AssetReference = InItem->GetSubCategoryObjectAsset();
-			FTopLevelAssetPath TopLevelAssetPath;
-
-			if (AssetReference.IsAsset())
-			{
-				TopLevelAssetPath = FTopLevelAssetPath(AssetReference.GetLongPackageFName(), *AssetReference.GetAssetName());
-			}
-
-			// First check pin type permissions
-			if (!FBlueprintActionDatabase::IsPinTypeAllowed(InItem->GetPinType(false), TopLevelAssetPath))
-			{
-				return false;
-			}
-
-			// Then asset permissions
-			if(AssetReferenceFilter.IsValid() && AssetReference.IsValid())
-			{
-				FAssetRegistryModule& AssetRegistryModule = FModuleManager::LoadModuleChecked<FAssetRegistryModule>("AssetRegistry");
-				FAssetData AssetData = AssetRegistryModule.Get().GetAssetByObjectPath(AssetReference.GetAssetPathName());
-				if (!AssetReferenceFilter->PassesFilter(AssetData))
-				{
-					return false;
-				}
-			}
-
-			return true;
-		}
-
-	private:
-		/** Filter for asset references */
-		TSharedPtr<IAssetReferenceFilter> AssetReferenceFilter;
-	};
-
 	/**
 	 * Utility function that will check to see if the specified graph has any 
 	 * nodes other than those that come default, pre-placed, in the graph.
@@ -980,12 +927,6 @@ void FBlueprintEditor::AnalyticsTrackCompileEvent( UBlueprint* Blueprint, int32
 	}
 }
 
-void FBlueprintEditor::GetPinTypeSelectorFilters(TArray<TSharedPtr<IPinTypeSelectorFilter>>& OutFilters) const
-{ 
-	OutFilters.Add(ImportedPinTypeSelectorFilter);
-	OutFilters.Add(PermissionsPinTypeSelectorFilter);
-}
-
 void FBlueprintEditor::RefreshEditors(ERefreshBlueprintEditorReason::Type Reason)
 {
 	bool bForceFocusOnSelectedNodes = false;
@@ -1949,7 +1890,6 @@ void FBlueprintEditor::CommonInitialization(const TArray<UBlueprint*>& InitBluep
 	// Create imported namespace type filters for value editing.
 	ImportedClassViewerFilter = MakeShared<BlueprintEditorImpl::FImportedClassViewerFilterProxy>(ImportedNamespaceHelper->GetClassViewerFilter());
 	ImportedPinTypeSelectorFilter = MakeShared<BlueprintEditorImpl::FImportedPinTypeSelectorFilterProxy>(ImportedNamespaceHelper->GetPinTypeSelectorFilter());
-	PermissionsPinTypeSelectorFilter = MakeShared<BlueprintEditorImpl::FPermissionsPinTypeSelectorFilter>(InitBlueprints);
 
 	// Make sure we know when tabs become active to update details tab
 	OnActiveTabChangedDelegateHandle = FGlobalTabmanager::Get()->OnActiveTabChanged_Subscribe( FOnActiveTabChanged::FDelegate::CreateRaw(this, &FBlueprintEditor::OnActiveTabChanged) );
diff --git a/Engine/Source/Editor/Kismet/Private/SBlueprintPalette.cpp b/Engine/Source/Editor/Kismet/Private/SBlueprintPalette.cpp
--- a/Engine/Source/Editor/Kismet/Private/SBlueprintPalette.cpp
+++ b/Engine/Source/Editor/Kismet/Private/SBlueprintPalette.cpp
@@ -754,10 +754,10 @@ private:
 	{
 		OnTypeChanged = InArgs._OnTypeChanged;
 		
-		TArray<TSharedPtr<IPinTypeSelectorFilter>> CustomPinTypeFilters;
+		TSharedPtr<IPinTypeSelectorFilter> CustomPinTypeFilter;
 		if (BlueprintEditorPtr.IsValid())
 		{
-			BlueprintEditorPtr.Pin()->GetPinTypeSelectorFilters(CustomPinTypeFilters);
+			CustomPinTypeFilter = BlueprintEditorPtr.Pin()->GetImportedPinTypeSelectorFilter();
 		}
 
 		const UEdGraphSchema* Schema = GetDefault<UEdGraphSchema_K2>();
@@ -784,7 +784,7 @@ private:
 				.OnPinTypeChanged(this, &SPinTypeSelectorHelper::OnVarTypeChanged)
 				.TypeTreeFilter(ETypeTreeFilter::None)
 				.SelectorType(BlueprintEditorPtr.IsValid() ? SPinTypeSelector::ESelectorType::Partial : SPinTypeSelector::ESelectorType::None)
-				.CustomFilters(CustomPinTypeFilters)
+				.CustomFilter(CustomPinTypeFilter)
 			];	
 		}
 	}
diff --git a/Engine/Source/Editor/Kismet/Public/BlueprintEditor.h b/Engine/Source/Editor/Kismet/Public/BlueprintEditor.h
--- a/Engine/Source/Editor/Kismet/Public/BlueprintEditor.h
+++ b/Engine/Source/Editor/Kismet/Public/BlueprintEditor.h
@@ -244,9 +244,7 @@ public:
 	virtual void AnalyticsTrackNodeEvent(UBlueprint* Blueprint, UEdGraphNode *GraphNode, bool bNodeDelete = false) const override;
 	void AnalyticsTrackCompileEvent(UBlueprint* Blueprint, int32 NumErrors, int32 NumWarnings) const;
 	virtual TSharedPtr<class IClassViewerFilter> GetImportedClassViewerFilter() const override { return ImportedClassViewerFilter; }
-	UE_DEPRECATED(5.1, "Please use GetPinTypeSelectorFilters")
 	virtual TSharedPtr<class IPinTypeSelectorFilter> GetImportedPinTypeSelectorFilter() const override { return ImportedPinTypeSelectorFilter; }
-	virtual void GetPinTypeSelectorFilters(TArray<TSharedPtr<class IPinTypeSelectorFilter>>& OutFilters) const override;
 	virtual bool IsNonImportedObject(const UObject* InObject) const;
 	virtual bool IsNonImportedObject(const FSoftObjectPath& InObject) const;
 	//~ End IBlueprintEditor Interface
@@ -1417,12 +1415,6 @@ protected:
 	/** Filter used to restrict pin type selector widgets in the editor context to imported namespaces only */
 	TSharedPtr<class IPinTypeSelectorFilter> ImportedPinTypeSelectorFilter;
 
-	/** Filter used to restrict pin type selector widgets to those which permissions allow */
-	TSharedPtr<class IPinTypeSelectorFilter> PermissionsPinTypeSelectorFilter;
-
-	/** Filter used to compose multiple pin type filters together */
-	TSharedPtr<class IPinTypeSelectorFilter> CompositePinTypeSelectorFilter;
-	
 	FOnSetPinVisibility OnSetPinVisibility;
 
 	/** Has someone requested a deferred update of the saved document state? */
diff --git a/Engine/Source/Editor/Kismet/Public/BlueprintEditorModule.h b/Engine/Source/Editor/Kismet/Public/BlueprintEditorModule.h
--- a/Engine/Source/Editor/Kismet/Public/BlueprintEditorModule.h
+++ b/Engine/Source/Editor/Kismet/Public/BlueprintEditorModule.h
@@ -109,12 +109,8 @@ public:
 	virtual TSharedPtr<class IClassViewerFilter> GetImportedClassViewerFilter() const { return nullptr; }
 
 	/** Return the pin type selector filter associated with the current set of imported namespaces within this editor context. Default is NULL (no filter). */
-	UE_DEPRECATED(5.1, "Please use GetPinTypeSelectorFilters")
 	virtual TSharedPtr<class IPinTypeSelectorFilter> GetImportedPinTypeSelectorFilter() const { return nullptr; }
 
-	/** Get all the the pin type selector filters within this editor context. */
-	virtual void GetPinTypeSelectorFilters(TArray<TSharedPtr<class IPinTypeSelectorFilter>>& OutFilters) const {}
-	
 	/** Return whether the given object falls outside the scope of the current set of imported namespaces within this editor context. Default is FALSE (imported). */
 	virtual bool IsNonImportedObject(const UObject* InObject) const { return false; }
 
diff --git a/Engine/Source/Editor/KismetWidgets/Private/SPinTypeSelector.cpp b/Engine/Source/Editor/KismetWidgets/Private/SPinTypeSelector.cpp
--- a/Engine/Source/Editor/KismetWidgets/Private/SPinTypeSelector.cpp
+++ b/Engine/Source/Editor/KismetWidgets/Private/SPinTypeSelector.cpp
@@ -314,22 +314,14 @@ void SPinTypeSelector::Construct(const FArguments& InArgs, FGetPinTypeTree GetPi
 	SelectorType = InArgs._SelectorType;
 
 	NumFilteredPinTypeItems = 0;
-
-	if(InArgs._CustomFilters.Num() > 0)
+	if (InArgs._CustomFilter.IsValid())
 	{
-		for(const TSharedPtr<IPinTypeSelectorFilter>& Filter : InArgs._CustomFilters)
-		{
-			CustomFilters.Add(MakeShared<FPinTypeSelectorCustomFilterProxy>(Filter.ToSharedRef(), FSimpleDelegate::CreateSP(this, &SPinTypeSelector::OnCustomFilterChanged)));
-		}
-	}
-	else if (InArgs._CustomFilter.IsValid())
-	{
-		CustomFilters.Add(MakeShared<FPinTypeSelectorCustomFilterProxy>(InArgs._CustomFilter.ToSharedRef(), FSimpleDelegate::CreateSP(this, &SPinTypeSelector::OnCustomFilterChanged)));
+		CustomFilter = MakeShared<FPinTypeSelectorCustomFilterProxy>(InArgs._CustomFilter.ToSharedRef(), FSimpleDelegate::CreateSP(this, &SPinTypeSelector::OnCustomFilterChanged));
 	}
 	else if (UClass* PinTypeSelectorFilterClass = GetDefault<UPinTypeSelectorFilter>()->FilterClass.LoadSynchronous() )
 	{
 		TSharedPtr<IPinTypeSelectorFilter> SelectorFilter = GetDefault<UPinTypeSelectorFilter>(PinTypeSelectorFilterClass)->GetPinTypeSelectorFilter();
-		CustomFilters.Add(MakeShared<FPinTypeSelectorCustomFilterProxy>(SelectorFilter.ToSharedRef(), FSimpleDelegate::CreateSP(this, &SPinTypeSelector::OnCustomFilterChanged)));
+		CustomFilter = MakeShared<FPinTypeSelectorCustomFilterProxy>(SelectorFilter.ToSharedRef(), FSimpleDelegate::CreateSP(this, &SPinTypeSelector::OnCustomFilterChanged));
 	}
 
 	bIsRightMousePressed = false;
@@ -1044,7 +1036,7 @@ TSharedRef<SWidget>	SPinTypeSelector::GetMenuContent(bool bForSecondaryType)
 	GetChildrenWithSupportedTypes(TypeTreeRoot, FilteredBySupportedTypes);
 	TypeTreeRoot = FilteredBySupportedTypes;
 
-	if (CustomFilters.Num() > 0)
+	if (CustomFilter.IsValid())
 	{
 		NumFilteredPinTypeItems = 0;
 		FilteredTypeTreeRoot.Empty();
@@ -1071,21 +1063,12 @@ TSharedRef<SWidget>	SPinTypeSelector::GetMenuContent(bool bForSecondaryType)
 			.OnTextChanged( this, &SPinTypeSelector::OnFilterTextChanged )
 			.OnTextCommitted( this, &SPinTypeSelector::OnFilterTextCommitted );
 
-		TArray<TSharedPtr<SWidget>> CustomFilterOptionsWidgets;
-		for(const TSharedPtr<IPinTypeSelectorFilter>& CustomFilter : CustomFilters)
+		TSharedPtr<SWidget> CustomFilterOptionsWidget;
+		if (CustomFilter.IsValid())
 		{
-			if(CustomFilter.IsValid())
-			{
-				TSharedPtr<SWidget> CustomWidget = CustomFilter->GetFilterOptionsWidget();
-				if(CustomWidget.IsValid())
-				{
-					CustomFilterOptionsWidgets.Add(CustomWidget);
-				}
-			}
+			CustomFilterOptionsWidget = CustomFilter->GetFilterOptionsWidget();
 		}
 
-		TSharedPtr<SHorizontalBox> CustomWidgetContainer;
-		
 		MenuContent = SAssignNew(PinTypeSelectorMenuOwner, SMenuOwner)
 			[
 				SNew(SListViewSelectorDropdownMenu<FPinTypeTreeItem>, FilterTextBox, TypeTreeView)
@@ -1113,9 +1096,9 @@ TSharedRef<SWidget>	SPinTypeSelector::GetMenuContent(bool bForSecondaryType)
 					.Padding(8.f, 0.f, 8.f, 4.f)
 					[
 						SNew(SBox)
-						.Visibility(CustomFilterOptionsWidgets.Num() > 0 ? EVisibility::Visible : EVisibility::Collapsed)
+						.Visibility(CustomFilter.IsValid() ? EVisibility::Visible : EVisibility::Collapsed)
 						[
-							SAssignNew(CustomWidgetContainer, SHorizontalBox)
+							SNew(SHorizontalBox)
 							+SHorizontalBox::Slot()
 							.VAlign(VAlign_Center)
 							.FillWidth(1.f)
@@ -1123,24 +1106,17 @@ TSharedRef<SWidget>	SPinTypeSelector::GetMenuContent(bool bForSecondaryType)
 								SNew(STextBlock)
 								.Text(this, &SPinTypeSelector::GetPinTypeItemCountText)
 							]
-							
+							+SHorizontalBox::Slot()
+							.VAlign(VAlign_Center)
+							.AutoWidth()
+							[
+								CustomFilterOptionsWidget.IsValid() ? CustomFilterOptionsWidget.ToSharedRef() : SNullWidget::NullWidget
+							]
 						]
 					]
 				]
 			];
-
-		if(CustomWidgetContainer.IsValid())
-		{
-			for(const TSharedPtr<SWidget>& CustomFilterOptionsWidget : CustomFilterOptionsWidgets)
-			{
-				CustomWidgetContainer->AddSlot()
-				.VAlign(VAlign_Center)
-				.AutoWidth()
-				[
-					CustomFilterOptionsWidget.ToSharedRef()
-				];
-			}
-		}
+			
 
 		if (bForSecondaryType)
 		{
@@ -1338,14 +1314,8 @@ bool SPinTypeSelector::GetChildrenMatchingSearch(const FText& InSearchText, cons
 		// If children match the search filter, there's no need to do any additional checks
 		if (!bHasChildrenMatchingSearch)
 		{
-			// If valid, attempt to match the custom filters; otherwise, treat it as a match by default.
-			for(const TSharedPtr<IPinTypeSelectorFilter>& CustomFilter : CustomFilters)
-			{
-				if(CustomFilter.IsValid())
-				{
-					bFilterMatches &= CustomFilter->ShouldShowPinTypeTreeItem(Item);
-				}
-			}
+			// If valid, attempt to match the custom filter; otherwise, treat it as a match by default.
+			bFilterMatches = CustomFilter.IsValid() ? CustomFilter->ShouldShowPinTypeTreeItem(Item) : true;
 
 			// If we didn't match the custom filter, or it's an empty search, let's not do any checks against the FilterTerms
 			if (bFilterMatches && !bIsEmptySearch)
diff --git a/Engine/Source/Editor/KismetWidgets/Public/SPinTypeSelector.h b/Engine/Source/Editor/KismetWidgets/Public/SPinTypeSelector.h
--- a/Engine/Source/Editor/KismetWidgets/Public/SPinTypeSelector.h
+++ b/Engine/Source/Editor/KismetWidgets/Public/SPinTypeSelector.h
@@ -82,8 +82,7 @@ public:
 		SLATE_ATTRIBUTE( FSlateFontInfo, Font )
 		SLATE_ARGUMENT( ESelectorType, SelectorType )
 		SLATE_ATTRIBUTE(bool, ReadOnly)
-		SLATE_ARGUMENT_DEPRECATED(TSharedPtr<class IPinTypeSelectorFilter>, CustomFilter, 5.1, "Please use CustomFilters instead")
-		SLATE_ARGUMENT(TArray<TSharedPtr<class IPinTypeSelectorFilter>>, CustomFilters)
+		SLATE_ARGUMENT(TSharedPtr<class IPinTypeSelectorFilter>, CustomFilter)
 	SLATE_END_ARGS()
 public:
 	void Construct(const FArguments& InArgs, FGetPinTypeTree GetPinTypeTreeFunc);
@@ -167,7 +166,7 @@ protected:
 	TWeakPtr<SMenuOwner> PinTypeSelectorMenuOwner;
 
 	/** An interface to optionally apply a custom filter to the available pin type items for display. */
-	TArray<TSharedPtr<class IPinTypeSelectorFilter>> CustomFilters;
+	TSharedPtr<class IPinTypeSelectorFilter> CustomFilter;
 
 	/** Array checkbox support functions */
 	ECheckBoxState IsArrayChecked() const;
