CreateBlueprintFromActor dialog fixes:
- Take the global class filter into account to determine whether Subclass mode is allowed
- Use the destination asset path to drive the asset reference filter (TODO: update the class viewer whenever the destination folder changes)

diff --git a/Engine/Source/Editor/KismetWidgets/Private/CreateBlueprintFromActorDialog.cpp b/Engine/Source/Editor/KismetWidgets/Private/CreateBlueprintFromActorDialog.cpp
--- a/Engine/Source/Editor/KismetWidgets/Private/CreateBlueprintFromActorDialog.cpp
+++ b/Engine/Source/Editor/KismetWidgets/Private/CreateBlueprintFromActorDialog.cpp
@@ -150,6 +150,17 @@ ECreateBlueprintFromActorMode FCreateBlueprintFromActorDialog::GetValidCreationM
 			if (NumSelectedActors == 0)
 			{
 				bCanSubclass = FKismetEditorUtilities::CanCreateBlueprintOfClass(Actor->GetClass());
+				if (bCanSubclass)
+				{
+					// Check whether the class is allowed by the global class filter
+					FClassViewerModule& ClassViewerModule = FModuleManager::LoadModuleChecked<FClassViewerModule>("ClassViewer");
+					if (const TSharedPtr<IClassViewerFilter>& GlobalClassFilter = ClassViewerModule.GetGlobalClassViewerFilter())
+					{
+						TSharedRef<FClassViewerFilterFuncs> ClassFilterFuncs = ClassViewerModule.CreateFilterFuncs();
+						FClassViewerInitializationOptions ClassViewerOptions = {};
+						bCanSubclass = GlobalClassFilter->IsClassAllowed(ClassViewerOptions, Actor->GetClass(), ClassFilterFuncs);
+					}
+				}
 			}
 
 			if (bCanCreatePrefab && Actor->GetClass()->HasAnyClassFlags(CLASS_NotPlaceable))
@@ -239,11 +250,57 @@ void SSCreateBlueprintPicker::Construct(const FArguments& InArgs)
 
 	bPressedOk = false;
 	ChosenClass = nullptr;
-	CreateMode = InArgs._CreateMode;
 
-	FContentBrowserModule& ContentBrowserModule = FModuleManager::LoadModuleChecked<FContentBrowserModule>("ContentBrowser");
-	FClassViewerModule& ClassViewerModule = FModuleManager::LoadModuleChecked<FClassViewerModule>("ClassViewer");
-	FAssetToolsModule& AssetToolsModule = FModuleManager::LoadModuleChecked<FAssetToolsModule>("AssetTools");
+	ECreateBlueprintFromActorMode ValidCreateMethods = FCreateBlueprintFromActorDialog::GetValidCreationMethods();
+	const bool bCanHarvestComponents = !!(ValidCreateMethods & ECreateBlueprintFromActorMode::Harvest);
+	const bool bCanSubclass = !!(ValidCreateMethods & ECreateBlueprintFromActorMode::Subclass);
+	const bool bCanCreatePrefab = !!(ValidCreateMethods & ECreateBlueprintFromActorMode::ChildActor);
+
+	if (!!(InArgs._CreateMode & ValidCreateMethods))
+	{
+		CreateMode = InArgs._CreateMode;
+	}
+	else
+	{
+		if (bCanSubclass)
+		{
+			CreateMode = ECreateBlueprintFromActorMode::Subclass;
+		}
+		else if (bCanCreatePrefab)
+		{
+			CreateMode = ECreateBlueprintFromActorMode::ChildActor;
+		}
+		else if (bCanHarvestComponents)
+		{
+			CreateMode = ECreateBlueprintFromActorMode::Harvest;
+		}
+		else
+		{
+			CreateMode = ECreateBlueprintFromActorMode::None;
+		}
+	}
+
+	// Set initial destination asset folder and name
+	{
+		FContentBrowserModule& ContentBrowserModule = FModuleManager::LoadModuleChecked<FContentBrowserModule>("ContentBrowser");
+
+		AssetPath = ContentBrowserModule.Get().GetCurrentPath();
+		// Change path if cannot write to it
+		AssetPath = ContentBrowserModule.Get().GetInitialPathToSaveAsset(AssetPath);
+
+		for (FSelectionIterator Iter(*GEditor->GetSelectedActors()); Iter; ++Iter)
+		{
+			AActor* Actor = Cast<AActor>(*Iter);
+			if (Actor)
+			{
+				AssetName += Actor->GetActorLabel();
+				AssetName += TEXT("_");
+				break;
+			}
+		}
+
+		AssetName = UPackageTools::SanitizePackageName(AssetName + TEXT("Blueprint"));
+	}
 
 	ActorOverride = InArgs._ActorOverride;
 
@@ -306,33 +363,18 @@ void SSCreateBlueprintPicker::Construct(const FArguments& InArgs)
 		ClassViewerOptions.InitiallySelectedClass = AActor::StaticClass();
 	}
 
-	ClassViewer = StaticCastSharedRef<SClassViewer>(ClassViewerModule.CreateClassViewer(ClassViewerOptions, FOnClassPicked::CreateSP(this, &SSCreateBlueprintPicker::OnClassPicked)));
-
-	AssetPath = ContentBrowserModule.Get().GetCurrentPath();
-
-	// Change path if cannot write to it
-	AssetPath = ContentBrowserModule.Get().GetInitialPathToSaveAsset(AssetPath);
-
-	ECreateBlueprintFromActorMode ValidCreateMethods = FCreateBlueprintFromActorDialog::GetValidCreationMethods();
-
-	const bool bCanHarvestComponents = !!(ValidCreateMethods & ECreateBlueprintFromActorMode::Harvest);
-	const bool bCanSubclass = !!(ValidCreateMethods & ECreateBlueprintFromActorMode::Subclass);
-	const bool bCanCreatePrefab = !!(ValidCreateMethods & ECreateBlueprintFromActorMode::ChildActor);
-
-	for (FSelectionIterator Iter(*GEditor->GetSelectedActors()); Iter; ++Iter)
 	{
-		AActor* Actor = Cast<AActor>(*Iter);
-		if (Actor)
-		{
-			AssetName += Actor->GetActorLabel();
-			AssetName += TEXT("_");
-			break;
-		}
+		const FString DestPackageName = FPaths::Combine(AssetPath.GetInternalPathString(), AssetName);
+		const FString DestAssetPath = FString::Printf(TEXT("%s.%s"), *DestPackageName, *AssetName);
+		ClassViewerOptions.AdditionalReferencingAssets.Add(FAssetData(DestPackageName, DestAssetPath, UBlueprint::StaticClass()->GetClassPathName()));
+		// @fixme: Update the class viewer whenever the destination folder changes
 	}
 
-	AssetName = UPackageTools::SanitizePackageName(AssetName + TEXT("Blueprint"));
+	FClassViewerModule& ClassViewerModule = FModuleManager::LoadModuleChecked<FClassViewerModule>("ClassViewer");
+	ClassViewer = StaticCastSharedRef<SClassViewer>(ClassViewerModule.CreateClassViewer(ClassViewerOptions, FOnClassPicked::CreateSP(this, &SSCreateBlueprintPicker::OnClassPicked)));
 
 	FString PackageName;
+	FAssetToolsModule& AssetToolsModule = FModuleManager::LoadModuleChecked<FAssetToolsModule>("AssetTools");
 	AssetToolsModule.Get().CreateUniqueAssetName(AssetPath.GetInternalPathString() / AssetName, TEXT(""), PackageName, AssetName);
 
 	TSharedPtr<SGridPanel> CreationMethodSection;
