EditorUX - PR #9588: Feature: Allow changing a BP variable's object reference type directly from the dropdown (Contributed by [REDACTED_PERSON])

diff --git a/Engine/Source/Editor/KismetWidgets/Private/SPinTypeSelector.cpp b/Engine/Source/Editor/KismetWidgets/Private/SPinTypeSelector.cpp
--- a/Engine/Source/Editor/KismetWidgets/Private/SPinTypeSelector.cpp
+++ b/Engine/Source/Editor/KismetWidgets/Private/SPinTypeSelector.cpp
@@ -844,7 +844,38 @@ void SPinTypeSelector::OnObjectReferenceSelectionChanged(FObjectReferenceListIte
 
 TSharedRef< SWidget > SPinTypeSelector::GetAllowedObjectTypes(FPinTypeTreeItem InItem, bool bForSecondaryType)
 {
-	AllowedObjectReferenceTypes.Reset();
+	GenerateAllowedObjectTypesList(AllowedObjectReferenceTypes, InItem, bForSecondaryType);
+
+	TSharedRef<SListView<FObjectReferenceListItem>> ListView = SNew(SListView<FObjectReferenceListItem>)
+		.ListItemsSource(&AllowedObjectReferenceTypes)
+		.SelectionMode(ESelectionMode::Single)
+		.OnGenerateRow(this, &SPinTypeSelector::GenerateObjectReferenceTreeRow)
+		.OnSelectionChanged(this, &SPinTypeSelector::OnObjectReferenceSelectionChanged, bForSecondaryType);
+
+	WeakListView = ListView;
+
+	if (AllowedObjectReferenceTypes.Num())
+	{
+		ListView->SetSelection(AllowedObjectReferenceTypes[0], ESelectInfo::OnNavigation);
+	}
+
+	return 
+		SNew(SBorder)
+		.BorderImage(FAppStyle::GetBrush("Menu.Background"))
+		[
+			SNew(SObjectReferenceWidget, PinTypeSelectorMenuOwner)
+			[
+				SNew(SListViewSelectorDropdownMenu<FObjectReferenceListItem>, nullptr, ListView)
+				[
+					ListView
+				]
+			]
+		];
+}
+
+void SPinTypeSelector::GenerateAllowedObjectTypesList(TArray<FObjectReferenceListItem>& OutList, FPinTypeTreeItem InItem, bool bForSecondaryType) const
+{
+	OutList.Reset();
 
 	// Do not force the pin type here, that causes a load of the Blueprint (if unloaded)
 	FEdGraphPinType PinType = InItem->GetPinType(false);
@@ -868,7 +899,7 @@ TSharedRef< SWidget > SPinTypeSelector::GetAllowedObjectTypes(FPinTypeTreeItem I
 		PinType.PinCategory = UEdGraphSchema_K2::PC_Object;
 		TSharedRef<SWidget> Widget = CreateObjectReferenceWidget(InItem, PinType, IconBrush, FText::Format(LOCTEXT("ObjectTooltip", "Reference an instanced object of type \'{TypeName}\'"), Args));
 		FObjectReferenceListItem ObjectReferenceType = MakeShareable(new FObjectReferenceType(InItem, Widget, PinType.PinCategory));
-		AllowedObjectReferenceTypes.Add(ObjectReferenceType);
+		OutList.Add(ObjectReferenceType);
 	}
 
 	if (PossibleObjectReferenceTypes & static_cast<uint8>(EObjectReferenceType::ClassReference))
@@ -876,7 +907,7 @@ TSharedRef< SWidget > SPinTypeSelector::GetAllowedObjectTypes(FPinTypeTreeItem I
 		PinType.PinCategory = UEdGraphSchema_K2::PC_Class;
 		TSharedRef<SWidget> Widget = CreateObjectReferenceWidget(InItem, PinType, IconBrush, FText::Format(LOCTEXT("ClassTooltip", "Reference a class inheriting from type \'{TypeName}\'"), Args));
 		FObjectReferenceListItem ObjectReferenceType = MakeShareable(new FObjectReferenceType(InItem, Widget, PinType.PinCategory));
-		AllowedObjectReferenceTypes.Add(ObjectReferenceType);
+		OutList.Add(ObjectReferenceType);
 	}
 
 	if (PossibleObjectReferenceTypes & static_cast<uint8>(EObjectReferenceType::SoftObject))
@@ -890,7 +921,7 @@ TSharedRef< SWidget > SPinTypeSelector::GetAllowedObjectTypes(FPinTypeTreeItem I
 		PinType.PinCategory = UEdGraphSchema_K2::PC_SoftObject;
 		TSharedRef<SWidget> Widget = CreateObjectReferenceWidget(InItem, PinType, IconBrush, FText::Format(FormatString, Args));
 		FObjectReferenceListItem ObjectReferenceType = MakeShareable(new FObjectReferenceType(InItem, Widget, PinType.PinCategory));
-		AllowedObjectReferenceTypes.Add(ObjectReferenceType);
+		OutList.Add(ObjectReferenceType);
 	}
 
 	if (PossibleObjectReferenceTypes & static_cast<uint8>(EObjectReferenceType::SoftClass))
@@ -904,34 +935,8 @@ TSharedRef< SWidget > SPinTypeSelector::GetAllowedObjectTypes(FPinTypeTreeItem I
 		PinType.PinCategory = UEdGraphSchema_K2::PC_SoftClass;
 		TSharedRef<SWidget> Widget = CreateObjectReferenceWidget(InItem, PinType, IconBrush, FText::Format(FormatString, Args));
 		FObjectReferenceListItem ObjectReferenceType = MakeShareable(new FObjectReferenceType(InItem, Widget, PinType.PinCategory));
-		AllowedObjectReferenceTypes.Add(ObjectReferenceType);
+		OutList.Add(ObjectReferenceType);
 	}
-
-	TSharedPtr<SListView<FObjectReferenceListItem>> ListView;
-	SAssignNew(ListView, SListView<FObjectReferenceListItem>)
-		.ListItemsSource(&AllowedObjectReferenceTypes)
-		.SelectionMode(ESelectionMode::Single)
-		.OnGenerateRow(this, &SPinTypeSelector::GenerateObjectReferenceTreeRow)
-		.OnSelectionChanged(this, &SPinTypeSelector::OnObjectReferenceSelectionChanged, bForSecondaryType);
-
-	WeakListView = ListView;
-	if (AllowedObjectReferenceTypes.Num())
-	{
-		ListView->SetSelection(AllowedObjectReferenceTypes[0], ESelectInfo::OnNavigation);
-	}
-
-	return 
-		SNew(SBorder)
-		.BorderImage(FAppStyle::GetBrush("Menu.Background"))
-		[
-			SNew(SObjectReferenceWidget, PinTypeSelectorMenuOwner)
-			[
-				SNew(SListViewSelectorDropdownMenu<FObjectReferenceListItem>, nullptr, ListView)
-				[
-					ListView.ToSharedRef()
-				]
-			]
-		];
 }
 
 void SPinTypeSelector::OnSelectPinType(FPinTypeTreeItem InItem, FName InPinCategory, bool bForSecondaryType)
@@ -1119,49 +1124,137 @@ TSharedRef<SWidget>	SPinTypeSelector::GetMenuContent(bool bForSecondaryType)
 		}
 
 		TSharedPtr<SHorizontalBox> CustomWidgetContainer;
-		
+		TSharedPtr<SVerticalBox> TreeWrapper;
+
 		MenuContent = SAssignNew(PinTypeSelectorMenuOwner, SMenuOwner)
 			[
-				SNew(SListViewSelectorDropdownMenu<FPinTypeTreeItem>, FilterTextBox, TypeTreeView)
+				SAssignNew(TreeWrapper, SVerticalBox)
+				+SVerticalBox::Slot()
+				.FillHeight(1.f)
 				[
-					SNew( SVerticalBox )
-					+SVerticalBox::Slot()
-					.AutoHeight()
-					.Padding(4.f, 4.f, 4.f, 4.f)
+					SNew(SListViewSelectorDropdownMenu<FPinTypeTreeItem>, FilterTextBox, TypeTreeView)
 					[
-						FilterTextBox.ToSharedRef()
-					]
-					+SVerticalBox::Slot()
+						SNew( SVerticalBox )
+						+SVerticalBox::Slot()
 						.AutoHeight()
 						.Padding(4.f, 4.f, 4.f, 4.f)
+						[
+							FilterTextBox.ToSharedRef()
+						]
+						+SVerticalBox::Slot()
+							.AutoHeight()
+							.Padding(4.f, 4.f, 4.f, 4.f)
+							[
+								SNew(SBox)
+								.HeightOverride(TreeViewHeight)
+								.WidthOverride(TreeViewWidth)
+								[
+									TypeTreeView.ToSharedRef()
+								]
+							]
+						+SVerticalBox::Slot()
+						.AutoHeight()
+						.Padding(8.f, 0.f, 8.f, 4.f)
 						[
 							SNew(SBox)
-							.HeightOverride(TreeViewHeight)
-							.WidthOverride(TreeViewWidth)
+							.Visibility(CustomFilterOptionsWidgets.Num() > 0 ? EVisibility::Visible : EVisibility::Collapsed)
 							[
-								TypeTreeView.ToSharedRef()
+								SAssignNew(CustomWidgetContainer, SHorizontalBox)
+								+SHorizontalBox::Slot()
+								.VAlign(VAlign_Center)
+								.FillWidth(1.f)
+								[
+									SNew(STextBlock)
+									.Text(this, &SPinTypeSelector::GetPinTypeItemCountText)
+								]
 							]
 						]
+					]
+				]
+			];
+
+		static auto FindTreeItemForSubCategoryObject = [](const TArray<FPinTypeTreeItem>& TreeItems, const UObject* SubCategoryObject) -> FPinTypeTreeItem
+		{
+			if (!SubCategoryObject)
+			{
+				return nullptr;
+			}
+
+			for (const FPinTypeTreeItem& TreeItemCategory : TreeItems)
+			{
+				if (TreeItemCategory->GetPinType(false).PinCategory == UEdGraphSchema_K2::AllObjectTypes)
+				{
+					for (const FPinTypeTreeItem& TreeItem : TreeItemCategory->Children)
+					{
+						// Don't force load subcategory object, since the current pin's class
+						// should already have been loaded if we get to this point anyway
+						const bool bForceLoadedSubCategoryObject = false;
+						if (TreeItem->GetPinType(bForceLoadedSubCategoryObject).PinSubCategoryObject == SubCategoryObject)
+						{
+							return TreeItem;
+						}
+					}
+
+					// Objects should only be under the AllObjectTypes category, so once we've found that no need to keep looking
+					break;
+				}
+			}
+
+			return nullptr;
+		};
+
+		// If this pin type is an object type, then try to find its tree item (i.e. one with the same SubCategoryObject),
+		// since we need a tree item to know which available reference types it has, to then build the list from
+		const UObject* CurrentSubObject = bForSecondaryType ? TargetPinType.Get().PinValueType.TerminalSubCategoryObject.Get() : TargetPinType.Get().PinSubCategoryObject.Get();
+		FPinTypeTreeItem CurrentPinInfo = FindTreeItemForSubCategoryObject(TypeTreeRoot, CurrentSubObject);
+
+		// Note: Don't allow changing the reference type from compact selectors if it's a map type,
+		// since it's a bit ambiguous whether you're changing the key or value type
+		if (CurrentPinInfo && !(SelectorType == ESelectorType::Compact && CurrentPinInfo->GetPinType(false).IsMap()))
+		{
+			GenerateAllowedObjectTypesList(CurrentPinAllowedObjectReferenceTypes, CurrentPinInfo, bForSecondaryType);
+			if (CurrentPinAllowedObjectReferenceTypes.Num() > 0)
+			{
+				TSharedRef<SListView<FObjectReferenceListItem>> ListView = SNew(SListView<FObjectReferenceListItem>)
+					.ListItemsSource(&CurrentPinAllowedObjectReferenceTypes)
+					.SelectionMode(ESelectionMode::Single)
+					.OnGenerateRow(this, &SPinTypeSelector::GenerateObjectReferenceTreeRow)
+					.OnSelectionChanged(this, &SPinTypeSelector::OnObjectReferenceSelectionChanged, bForSecondaryType);
+
+				TreeWrapper->AddSlot()
+				.AutoHeight()
+				[
+					SNew(SVerticalBox)
+					.Visibility_Lambda([this](){ return SearchText.IsEmpty() ? EVisibility::Visible : EVisibility::Collapsed; })
+					+SVerticalBox::Slot()
+					.AutoHeight()
+					.Padding(6.f, 4.f, 4.f, 2.f)
+					[
+						SNew(STextBlock)
+							.Text(LOCTEXT("PinTypeSelector_ChangeReferenceTypeHeading", "Change to:"))
+							.Font(FAppStyle::GetFontStyle(TEXT("Kismet.TypePicker.CategoryFont")))
+							.ColorAndOpacity(FSlateColor::UseForeground())
+					]
 					+SVerticalBox::Slot()
 					.AutoHeight()
-					.Padding(8.f, 0.f, 8.f, 4.f)
+					.Padding(4.f, 0.f, 4.f, 4.f)
 					[
-						SNew(SBox)
-						.Visibility(CustomFilterOptionsWidgets.Num() > 0 ? EVisibility::Visible : EVisibility::Collapsed)
+						// This is the same setup/styling as GetAllowedObjectTypes()
+						SNew(SBorder)
+						.BorderImage(FAppStyle::GetBrush("Menu.Background"))
 						[
-							SAssignNew(CustomWidgetContainer, SHorizontalBox)
-							+SHorizontalBox::Slot()
-							.VAlign(VAlign_Center)
-							.FillWidth(1.f)
+							SNew(SObjectReferenceWidget, PinTypeSelectorMenuOwner)
 							[
-								SNew(STextBlock)
-								.Text(this, &SPinTypeSelector::GetPinTypeItemCountText)
+								SNew(SListViewSelectorDropdownMenu<FObjectReferenceListItem>, nullptr, ListView)
+								[
+									ListView
+								]
 							]
-							
 						]
 					]
-				]
-			];
+				];
+			}
+		}
 
 		if(CustomWidgetContainer.IsValid())
 		{
diff --git a/Engine/Source/Editor/KismetWidgets/Public/SPinTypeSelector.h b/Engine/Source/Editor/KismetWidgets/Public/SPinTypeSelector.h
--- a/Engine/Source/Editor/KismetWidgets/Public/SPinTypeSelector.h
+++ b/Engine/Source/Editor/KismetWidgets/Public/SPinTypeSelector.h
@@ -91,7 +91,7 @@ public:
 		, _TypeTreeFilter(ETypeTreeFilter::None)
 		, _bAllowArrays(true)
 		, _TreeViewWidth(300.f)
-		, _TreeViewHeight(400.f)
+		, _TreeViewHeight(350.f)
 		, _Font(FAppStyle::GetFontStyle(TEXT("NormalFont")))
 		, _SelectorType(ESelectorType::Full)
 		, _ReadOnly(false)
@@ -192,6 +192,9 @@ protected:
 	TWeakPtr<SListView<FObjectReferenceListItem>> WeakListView;
 	TWeakPtr<SMenuOwner> PinTypeSelectorMenuOwner;
 
+	/** Holds a cache of the allowed Object Reference types for the current pin type, to be shown inline. */
+	TArray<FObjectReferenceListItem> CurrentPinAllowedObjectReferenceTypes;
+
 	/** An interface to optionally apply a custom filter to the available pin type items for display. */
 	TArray<TSharedPtr<class IPinTypeSelectorFilter>> CustomFilters;
 
@@ -262,6 +265,9 @@ protected:
 
 	/** Gets the allowable object types for an tree item, used for building the sub-menu */
 	TSharedRef< SWidget > GetAllowedObjectTypes(FPinTypeTreeItem InItem, bool bForSecondaryType);
+
+	/** Fills the provided array with generated list items based on a pin type's allowed reference types, which could then be used as the item source for a list view */
+	void GenerateAllowedObjectTypesList(TArray<FObjectReferenceListItem>& OutList, FPinTypeTreeItem InItem, bool bForSecondaryType) const;
 	
 	/**
 	 * When a pin type is selected, handle it
