Editor UX - Using Levenshtein distance in the pin type selector dropdown now for blueprints.

diff --git a/Engine/Source/Editor/KismetWidgets/Private/SPinTypeSelector.cpp b/Engine/Source/Editor/KismetWidgets/Private/SPinTypeSelector.cpp
--- a/Engine/Source/Editor/KismetWidgets/Private/SPinTypeSelector.cpp
+++ b/Engine/Source/Editor/KismetWidgets/Private/SPinTypeSelector.cpp
@@ -1,6 +1,7 @@
 // Copyright Epic Games, Inc. All Rights Reserved.
 #include "SPinTypeSelector.h"
 
+#include "Algo/LevenshteinDistance.h"
 #include "BlueprintEditorSettings.h"
 #include "Containers/UnrealString.h"
 #include "EdGraph/EdGraphSchema.h"
@@ -1087,7 +1088,8 @@ TSharedRef<SWidget>	SPinTypeSelector::GetMenuContent(bool bForSecondaryType)
 		NumFilteredPinTypeItems = 0;
 		FilteredTypeTreeRoot.Empty();
 
-		GetChildrenMatchingSearch(FText::GetEmpty(), TypeTreeRoot, FilteredTypeTreeRoot);
+		FTopLevenshteinResult<FPinTypeTreeItem> TopLevenshteinResult;
+		GetChildrenMatchingSearch(FText::GetEmpty(), TypeTreeRoot, FilteredTypeTreeRoot, TopLevenshteinResult);
 	}
 	else
 	{
@@ -1349,6 +1351,48 @@ TSharedRef<SWidget> SPinTypeSelector::GetPinContainerTypeMenuContent()
 	return MenuBuilder.MakeWidget();
 }
 
+template<typename T>
+TOptional<T> FindBestMatch(FStringView SearchText, TArrayView<T> Container, TFunctionRef<FString(const T& ContainerItem)> ContainerElementTextExtractor)
+{
+	if (SearchText.IsEmpty())
+	{
+		return TOptional<T>();
+	}
+
+	if (Container.Num() == 0)
+	{
+		return TOptional<T>();
+	}
+
+	// Multiple candidates, pick one that is closest match.
+	auto CalculateScore = [](FStringView SearchValue, FStringView CandidateValue)
+	{
+		if (CandidateValue.IsEmpty())
+		{
+			return 0.0f;
+		}
+		const float WorstCase = static_cast<float>(SearchValue.Len() + CandidateValue.Len());
+		return 1.0f - (Algo::LevenshteinDistance(SearchValue, CandidateValue) / WorstCase);
+	};
+
+	const FString ObjectNameLowerCase = FString(SearchText).ToLower();
+
+	int32 HighestScoreIndex = 0;
+	float HighestScore = CalculateScore(ObjectNameLowerCase, ContainerElementTextExtractor(Container[0]).ToLower());
+
+	for (int32 Index = 1; Index < Container.Num(); Index++)
+	{
+		const float Score = CalculateScore(ObjectNameLowerCase, ContainerElementTextExtractor(Container[Index]).ToLower());
+		if (Score > HighestScore)
+		{
+			HighestScore = Score;
+			HighestScoreIndex = Index;
+		}
+	}
+
+	return Container[HighestScoreIndex];
+}
+
 //=======================================================================
 // Search Support
 void SPinTypeSelector::OnFilterTextChanged(const FText& NewText)
@@ -1357,31 +1401,23 @@ void SPinTypeSelector::OnFilterTextChanged(const FText& NewText)
 	NumFilteredPinTypeItems = 0;
 	FilteredTypeTreeRoot.Empty();
 
-	GetChildrenMatchingSearch(NewText, TypeTreeRoot, FilteredTypeTreeRoot);
+	FTopLevenshteinResult<FPinTypeTreeItem> TopLevenshteinResult;
+	GetChildrenMatchingSearch(NewText, TypeTreeRoot, FilteredTypeTreeRoot, TopLevenshteinResult);
 	TypeTreeView->RequestTreeRefresh();
-
-	// Select the first non-category item
-	auto SelectedItems = TypeTreeView->GetSelectedItems();
-	if(FilteredTypeTreeRoot.Num() > 0)
+	
+	if (TopLevenshteinResult.IsSet())
 	{
-		// Categories have children, we don't want to select categories
-		if(FilteredTypeTreeRoot[0]->Children.Num() > 0)
-		{
-			TypeTreeView->SetSelection(FilteredTypeTreeRoot[0]->Children[0], ESelectInfo::OnNavigation);
-		}
-		else
-		{
-			TypeTreeView->SetSelection(FilteredTypeTreeRoot[0], ESelectInfo::OnNavigation);
-		}
+		TypeTreeView->SetSelection(TopLevenshteinResult.Item, ESelectInfo::OnNavigation);
+		TypeTreeView->RequestScrollIntoView(TopLevenshteinResult.Item);
 	}
 }
 
 void SPinTypeSelector::OnFilterTextCommitted(const FText& NewText, ETextCommit::Type CommitInfo)
 {
-	if(CommitInfo == ETextCommit::OnEnter)
+	if (CommitInfo == ETextCommit::OnEnter)
 	{
 		auto SelectedItems = TypeTreeView->GetSelectedItems();
-		if(SelectedItems.Num() > 0)
+		if (SelectedItems.Num() > 0)
 		{
 			TypeTreeView->SetSelection(SelectedItems[0]);
 		}
@@ -1424,8 +1460,9 @@ bool SPinTypeSelector::GetChildrenWithSupportedTypes(const TArray<FPinTypeTreeIt
 	return bReturnVal;
 }
 
-bool SPinTypeSelector::GetChildrenMatchingSearch(const FText& InSearchText, const TArray<FPinTypeTreeItem>& UnfilteredList, TArray<FPinTypeTreeItem>& OutFilteredList)
+bool SPinTypeSelector::GetChildrenMatchingSearch(const FText& InSearchText, const TArray<FPinTypeTreeItem>& UnfilteredList, TArray<FPinTypeTreeItem>& OutFilteredList, FTopLevenshteinResult<FPinTypeTreeItem>& OutTopLevenshteinResult)
 {
+	FString TrimmedFilterString;
 	TArray<FString> FilterTerms;
 	TArray<FString> SanitizedFilterTerms;
 
@@ -1433,7 +1470,7 @@ bool SPinTypeSelector::GetChildrenMatchingSearch(const FText& InSearchText, cons
 	if (!bIsEmptySearch)
 	{
 		// Trim and sanitized the filter text (so that it more likely matches the action descriptions)
-		FString TrimmedFilterString = FText::TrimPrecedingAndTrailing(InSearchText).ToString();
+		TrimmedFilterString = FText::TrimPrecedingAndTrailing(InSearchText).ToString();
 
 		// Tokenize the search box text into a set of terms; all of them must be present to pass the filter
 		TrimmedFilterString.ParseIntoArray(FilterTerms, TEXT(" "), true);
@@ -1458,7 +1495,7 @@ bool SPinTypeSelector::GetChildrenMatchingSearch(const FText& InSearchText, cons
 		FPinTypeTreeItem NewInfo = MakeShareable( new UEdGraphSchema_K2::FPinTypeTreeInfo(Item) );
 		TArray<FPinTypeTreeItem> ValidChildren;
 
-		const bool bHasChildrenMatchingSearch = GetChildrenMatchingSearch(InSearchText, Item->Children, ValidChildren);
+		const bool bHasChildrenMatchingSearch = GetChildrenMatchingSearch(InSearchText, Item->Children, ValidChildren, OutTopLevenshteinResult);
 		bool bFilterMatches = bIsEmptySearch;
 
 		// If children match the search filter, there's no need to do any additional checks
@@ -1503,6 +1540,12 @@ bool SPinTypeSelector::GetChildrenMatchingSearch(const FText& InSearchText, cons
 			NewInfo->Children = ValidChildren;
 			OutFilteredList.Add(NewInfo);
 
+			// Skip anything with children, those are categories and we don't care about those.
+			if (Item->Children.Num() == 0)
+			{
+				OutTopLevenshteinResult.CompareAndUpdate(TrimmedFilterString, NewInfo, Item->GetDescription().ToString());
+			}
+
 			if (TypeTreeView.IsValid())
 			{
 				TypeTreeView->SetItemExpansion(NewInfo, !bIsEmptySearch);
@@ -1665,7 +1708,9 @@ void SPinTypeSelector::OnCustomFilterChanged()
 {
 	NumFilteredPinTypeItems = 0;
 	FilteredTypeTreeRoot.Empty();
-	GetChildrenMatchingSearch(SearchText, TypeTreeRoot, FilteredTypeTreeRoot);
+
+	FTopLevenshteinResult<FPinTypeTreeItem> TopLevenshteinResult;
+	GetChildrenMatchingSearch(SearchText, TypeTreeRoot, FilteredTypeTreeRoot, TopLevenshteinResult);
 
 	if (TypeTreeView.IsValid())
 	{
diff --git a/Engine/Source/Editor/KismetWidgets/Public/SPinTypeSelector.h b/Engine/Source/Editor/KismetWidgets/Public/SPinTypeSelector.h
--- a/Engine/Source/Editor/KismetWidgets/Public/SPinTypeSelector.h
+++ b/Engine/Source/Editor/KismetWidgets/Public/SPinTypeSelector.h
@@ -35,6 +35,7 @@
 #include "Widgets/Views/STableRow.h"
 #include "Widgets/Views/STableViewBase.h"
 #include "Widgets/Views/STreeView.h"
+#include "Algo/LevenshteinDistance.h"
 
 class ITableRow;
 class SComboButton;
@@ -51,6 +52,31 @@ struct FSlateBrush;
 
 DECLARE_DELEGATE_OneParam(FOnPinTypeChanged, const FEdGraphPinType&)
 
+template <typename ItemType>
+struct FTopLevenshteinResult
+{
+	ItemType Item;
+	float Score = INDEX_NONE;
+
+	bool IsSet() const { return Score != INDEX_NONE; }
+
+	void CompareAndUpdate(FStringView SearchValue, const ItemType& NewItem, FStringView NewItemValue)
+	{
+		if (SearchValue.IsEmpty() || NewItemValue.IsEmpty())
+		{
+			return;
+		}
+
+		const float WorstCase = static_cast<float>(SearchValue.Len() + NewItemValue.Len());
+		const float NormalizedDistance = 1.0f - (Algo::LevenshteinDistance(SearchValue, NewItemValue) / WorstCase);
+		if (NormalizedDistance > Score)
+		{
+			Score = NormalizedDistance;
+			Item = NewItem;
+		}
+	}
+};
+
 
 //////////////////////////////////////////////////////////////////////////
 // SPinTypeSelector
@@ -236,7 +262,7 @@ protected:
 	bool GetChildrenWithSupportedTypes(const TArray<FPinTypeTreeItem>& UnfilteredList, TArray<FPinTypeTreeItem>& OutFilteredList);
 
 	/** Helper to generate the filtered list of types, based on the search string matching */
-	bool GetChildrenMatchingSearch(const FText& SearchText, const TArray<FPinTypeTreeItem>& UnfilteredList, TArray<FPinTypeTreeItem>& OutFilteredList);
+	bool GetChildrenMatchingSearch(const FText& SearchText, const TArray<FPinTypeTreeItem>& UnfilteredList, TArray<FPinTypeTreeItem>& OutFilteredList, FTopLevenshteinResult<FPinTypeTreeItem>& OutTopLevenshteinResult);
 
 	/** Callback to get the tooltip text for the pin type combo box */
 	FText GetToolTipForComboBoxType() const;
