Cleanup pin 'custom filtering' logic, remove superfluous template

diff --git a/Engine/Source/Editor/KismetWidgets/Private/SPinTypeSelector.cpp b/Engine/Source/Editor/KismetWidgets/Private/SPinTypeSelector.cpp
--- a/Engine/Source/Editor/KismetWidgets/Private/SPinTypeSelector.cpp
+++ b/Engine/Source/Editor/KismetWidgets/Private/SPinTypeSelector.cpp
@@ -1400,30 +1400,45 @@ void SPinTypeSelector::OnFilterTextCommitted(const FText& NewText, ETextCommit::
 // helpers for dealing with the UEdGraphSchema_K2::FPinTypeTreeInfo tree:
 namespace PinTypeSelectorImpl
 {
-	template<typename T>
-	static void FilterUnsupportedTypesImpl(TArray<FPinTypeTreeItem>& ToFilter, const UEdGraphSchema* Schema, TWeakPtr<const FEdGraphSchemaAction> SchemaAction, T CustomFilter);
-	int32 CountValidTreeItemsImpl(const TArray<FPinTypeTreeItem>& ToCount);
+	static void FilterUnsupportedTypesImpl(TArray<FPinTypeTreeItem>& ToFilter, const UEdGraphSchema* Schema, TWeakPtr<const FEdGraphSchemaAction> SchemaAction, const TArray<TSharedPtr<class IPinTypeSelectorFilter>>& CustomFilters);
+	static int32 CountValidTreeItemsImpl(const TArray<FPinTypeTreeItem>& ToCount);
 }
 
-template<typename T>
-static void PinTypeSelectorImpl::FilterUnsupportedTypesImpl(TArray<FPinTypeTreeItem>& ToFilter, const UEdGraphSchema* Schema, TWeakPtr<const FEdGraphSchemaAction> SchemaAction, T CustomFilter)
+static void PinTypeSelectorImpl::FilterUnsupportedTypesImpl(TArray<FPinTypeTreeItem>& ToFilter, const UEdGraphSchema* Schema, TWeakPtr<const FEdGraphSchemaAction> SchemaAction, const TArray<TSharedPtr<class IPinTypeSelectorFilter>>& CustomFilters)
 {
-	ToFilter.SetNum(Algo::StableRemoveIf(ToFilter,
-		[Schema, SchemaAction, &CustomFilter = CustomFilter](const FPinTypeTreeItem& Item)
+	const auto FilterFunction = [Schema, SchemaAction, &CustomFilters = CustomFilters](const FPinTypeTreeItem& Item)
+	{
+		const auto DoCustomFilter = [&CustomFilters = CustomFilters](const FPinTypeTreeItem& Item)
 		{
-			if (Item->Children.Num() > 0)
+			bool bCustomFilterMatches = true;
+			for (const TSharedPtr<IPinTypeSelectorFilter>& CustomFilter : CustomFilters)
 			{
-				FilterUnsupportedTypesImpl<T>(Item->Children, Schema, SchemaAction, CustomFilter);
+				if (CustomFilter.IsValid())
+				{
+					bCustomFilterMatches &= CustomFilter->ShouldShowPinTypeTreeItem(Item);
+				}
 			}
-			const bool bSupportsType = (Item->Children.Num() != 0 ||
-				(	Schema->SupportsPinType(SchemaAction, Item->GetPinTypeNoResolve()) &&
-					CustomFilter(Item)));
-			return !bSupportsType;
+			return bCustomFilterMatches;
+		};
+
+		const bool bSupportsType = (Item->Children.Num() != 0 ||
+			(Schema->SupportsPinType(SchemaAction, Item->GetPinTypeNoResolve()) &&
+			 DoCustomFilter(Item)));
+		return !bSupportsType;
+	};
+
+	for (const FPinTypeTreeItem& Item : ToFilter)
+	{
+		if (Item->Children.Num() > 0)
+		{
+			FilterUnsupportedTypesImpl(Item->Children, Schema, SchemaAction, CustomFilters);
 		}
-	));
+	}
+
+	ToFilter.SetNum(Algo::StableRemoveIf(ToFilter, FilterFunction));
 }
 
-int32 PinTypeSelectorImpl::CountValidTreeItemsImpl(const TArray<FPinTypeTreeItem>& ToCount)
+static int32 PinTypeSelectorImpl::CountValidTreeItemsImpl(const TArray<FPinTypeTreeItem>& ToCount)
 {
 	int32 Count = Algo::CountIf(ToCount, [](const FPinTypeTreeItem& Item) { return !Item->bReadOnly; });
 	for (const FPinTypeTreeItem& Item : ToCount)
@@ -1452,32 +1467,7 @@ void SPinTypeSelector::FilterUnsupportedTypes(TArray<FPinTypeTreeItem>& ToFilter
 		}
 	}
 
-	const bool bHasCustomFilters = Algo::AnyOf(CustomFilters,
-		[](const TSharedPtr<IPinTypeSelectorFilter>& Filter)
-		{
-			return Filter.IsValid();
-		});
-
-	if (bHasCustomFilters)
-	{
-		PinTypeSelectorImpl::FilterUnsupportedTypesImpl(ToFilter, Schema, SchemaAction,
-			[&CustomFilters = CustomFilters](const FPinTypeTreeItem& Item)
-			{
-				bool bCustomFilterMatches = true;
-				for (const TSharedPtr<IPinTypeSelectorFilter>& CustomFilter : CustomFilters)
-				{
-					if (CustomFilter.IsValid())
-					{
-						bCustomFilterMatches &= CustomFilter->ShouldShowPinTypeTreeItem(Item);
-					}
-				}
-				return bCustomFilterMatches;
-			});
-	}
-	else
-	{
-		PinTypeSelectorImpl::FilterUnsupportedTypesImpl(ToFilter, Schema, SchemaAction, [](const FPinTypeTreeItem& Item) { return true; });
-	}
+	PinTypeSelectorImpl::FilterUnsupportedTypesImpl(ToFilter, Schema, SchemaAction, CustomFilters);
 
 	// count items after filtering invalid ones, set up the 'filtered view' to default to all the valid items:
 	NumValidPinTypeItems = PinTypeSelectorImpl::CountValidTreeItemsImpl(ToFilter);
