Fix CIS Issue 446379: type conversion warning

diff --git a/Engine/Source/Editor/KismetWidgets/Public/SPinTypeSelector.h b/Engine/Source/Editor/KismetWidgets/Public/SPinTypeSelector.h
--- a/Engine/Source/Editor/KismetWidgets/Public/SPinTypeSelector.h
+++ b/Engine/Source/Editor/KismetWidgets/Public/SPinTypeSelector.h
@@ -68,7 +68,7 @@ struct FTopLevenshteinResult
 		}
 
 		const float WorstCase = static_cast<float>(SearchValue.Len() + NewItemValue.Len());
-		const float NormalizedDistance = 1.0f - (Algo::LevenshteinDistance(SearchValue, NewItemValue) / WorstCase);
+		const float NormalizedDistance = 1.0f - static_cast<float>(Algo::LevenshteinDistance(SearchValue, NewItemValue)) / WorstCase;
 		if (NormalizedDistance > Score)
 		{
 			Score = NormalizedDistance;
