'Change To' menu does not show up when a object variable is first added in UAF
#fix update SPinSelector ContextMenu widget when underlying variable type has changed, as the type can cause different widgets to be constructed
#Synced-CL 48351162

diff --git a/Engine/Source/Editor/KismetWidgets/Private/SPinTypeSelector.cpp b/Engine/Source/Editor/KismetWidgets/Private/SPinTypeSelector.cpp
--- a/Engine/Source/Editor/KismetWidgets/Private/SPinTypeSelector.cpp
+++ b/Engine/Source/Editor/KismetWidgets/Private/SPinTypeSelector.cpp
@@ -1071,7 +1071,7 @@ TSharedRef<SWidget>	SPinTypeSelector::GetMenuContent(bool bForSecondaryType)
 	// Remove types not supported by schema
 	FilterUnsupportedTypes(TypeTreeRoot);
 
-	if( !MenuContent.IsValid() || (bForSecondaryType != bMenuContentIsSecondary) )
+	if( !MenuContent.IsValid() || (bForSecondaryType != bMenuContentIsSecondary) || (MenuContextPinType != TargetPinType.Get()) )
 	{
 		bMenuContentIsSecondary = bForSecondaryType;
 		// Pre-build the tree view and search box as it is needed as a parameter for the context menu's container.
@@ -1102,6 +1102,7 @@ TSharedRef<SWidget>	SPinTypeSelector::GetMenuContent(bool bForSecondaryType)
 		TSharedPtr<SHorizontalBox> CustomWidgetContainer;
 		TSharedPtr<SVerticalBox> TreeWrapper;
 
+		MenuContextPinType = TargetPinType.Get();
 		MenuContent = SAssignNew(PinTypeSelectorMenuOwner, SMenuOwner)
 			[
 				SAssignNew(TreeWrapper, SVerticalBox)
diff --git a/Engine/Source/Editor/KismetWidgets/Public/SPinTypeSelector.h b/Engine/Source/Editor/KismetWidgets/Public/SPinTypeSelector.h
--- a/Engine/Source/Editor/KismetWidgets/Public/SPinTypeSelector.h
+++ b/Engine/Source/Editor/KismetWidgets/Public/SPinTypeSelector.h
@@ -255,6 +255,8 @@ protected:
 
 	/** Reference to the menu content that's displayed when the type button is clicked on */
 	TSharedPtr<SMenuOwner> MenuContent;
+	FEdGraphPinType MenuContextPinType;
+	 
 	UE_API virtual TSharedRef<SWidget>	GetMenuContent(bool bForSecondaryType);
 	UE_API TSharedRef<SWidget> GetPinContainerTypeMenuContent();
 
