Slate Dynamic Invalidation - Graph Editor - Adding TypeTraits to trivial `SLevelOfDetailBranchNode`

#Synced-CL 47563119

diff --git a/Engine/Source/Editor/KismetWidgets/Public/SLevelOfDetailBranchNode.h b/Engine/Source/Editor/KismetWidgets/Public/SLevelOfDetailBranchNode.h
--- a/Engine/Source/Editor/KismetWidgets/Public/SLevelOfDetailBranchNode.h
+++ b/Engine/Source/Editor/KismetWidgets/Public/SLevelOfDetailBranchNode.h
@@ -20,6 +20,13 @@ struct FGeometry;
 
 DECLARE_DELEGATE_RetVal_OneParam(TSharedRef<SWidget>, FOnGetActiveDetailSlotContent, bool);
 
+
+template <>
+struct TWidgetTypeTraits<class SLevelOfDetailBranchNode>
+{
+	static constexpr bool SupportsInvalidation() { return true; }
+};
+
 class SLevelOfDetailBranchNode : public SCompoundWidget
 {
 	SLATE_BEGIN_ARGS(SLevelOfDetailBranchNode)
