Forcefully clear search text when selecting different pin type (this was preventing 'Change To' menu to show after searching for a UObject type, as search text != empty afterwards and only able to be cleared by hitting enter on empty searchbox)
#Synced-CL 48606655

diff --git a/Engine/Source/Editor/KismetWidgets/Private/SPinTypeSelector.cpp b/Engine/Source/Editor/KismetWidgets/Private/SPinTypeSelector.cpp
--- a/Engine/Source/Editor/KismetWidgets/Private/SPinTypeSelector.cpp
+++ b/Engine/Source/Editor/KismetWidgets/Private/SPinTypeSelector.cpp
@@ -1007,6 +1007,8 @@ void SPinTypeSelector::OnSelectPinType(FPinTypeTreeItem InItem, FName InPinCateg
 	{
 		NewTargetPinType.ContainerType = EPinContainerType::None;
 	}
+	
+	SearchText = FText::GetEmpty();
 
 	OnTypeChanged.ExecuteIfBound(NewTargetPinType);
 }
@@ -1083,6 +1085,7 @@ TSharedRef<SWidget>	SPinTypeSelector::GetMenuContent(bool bForSecondaryType)
 			.OnGetChildren(this, &SPinTypeSelector::GetTypeChildren);
 
 		SAssignNew(FilterTextBox, SSearchBox)
+			.InitialText(SearchText)
 			.OnTextChanged( this, &SPinTypeSelector::OnFilterTextChanged )
 			.OnTextCommitted( this, &SPinTypeSelector::OnFilterTextCommitted );
 
